Editor - Fix spelling of ContentBrowserFileDataSource folder
#rnx

diff --git a/Engine/Plugins/Editor/ContentBrowser/ContentBrowserAliasDataSource/ContentBrowserAliasDataSource.uplugin b/Engine/Plugins/Editor/ContentBrowser/ContentBrowserAliasDataSource/ContentBrowserAliasDataSource.uplugin
new file mode 100644
--- /dev/null
+++ b/Engine/Plugins/Editor/ContentBrowser/ContentBrowserAliasDataSource/ContentBrowserAliasDataSource.uplugin
@@ -0,0 +1,30 @@
+{
+	"FileVersion": 3,
+	"Version": 1,
+	"VersionName": "1.0",
+	"FriendlyName": "Content Browser - Alias Data Source",
+	"Description": "Data Source plugin providing allowing Content Browser items to appear in other directories other than their original location",
+	"Category": "Content Browser",
+	"CreatedBy": "Epic Games, Inc.",
+	"CreatedByURL": "https://epicgames.com",
+	"DocsURL": "",
+	"MarketplaceURL": "",
+	"SupportURL": "",
+	"EnabledByDefault": false,
+	"CanContainContent": false,
+	"IsBetaVersion": true,
+	"Installed": false,
+	"Modules": [
+		{
+			"Name": "ContentBrowserAliasDataSource",
+			"Type": "Editor",
+			"LoadingPhase": "Default"
+		}
+	],
+	"Plugins": [
+		{
+			"Name": "ContentBrowserAssetDataSource",
+			"Enabled": true
+		}
+	]
+}
\ No newline at end of file
diff --git a/Engine/Plugins/Editor/ContentBrowser/ContentBrowserAliasDataSource/Source/ContentBrowserAliasDataSource/ContentBrowserAliasDataSource.Build.cs b/Engine/Plugins/Editor/ContentBrowser/ContentBrowserAliasDataSource/Source/ContentBrowserAliasDataSource/ContentBrowserAliasDataSource.Build.cs
new file mode 100644
--- /dev/null
+++ b/Engine/Plugins/Editor/ContentBrowser/ContentBrowserAliasDataSource/Source/ContentBrowserAliasDataSource/ContentBrowserAliasDataSource.Build.cs
@@ -0,0 +1,30 @@
+// Copyright Epic Games, Inc. All Rights Reserved.
+
+namespace UnrealBuildTool.Rules
+{
+	public class ContentBrowserAliasDataSource : ModuleRules
+	{
+		public ContentBrowserAliasDataSource(ReadOnlyTargetRules Target) : base(Target)
+		{
+			ShortName = "CBADS";
+
+			PublicDependencyModuleNames.AddRange(
+				new string[] {
+					"AssetRegistry",
+					"Core",
+					"CoreUObject",
+					"ContentBrowserAssetDataSource",
+					"ContentBrowserData"
+				}
+			);
+
+			PrivateDependencyModuleNames.AddRange(
+				new string[]
+				{
+					"AssetTools",
+					"Engine"
+				}
+			);
+		}
+	}
+}
diff --git a/Engine/Plugins/Editor/ContentBrowser/ContentBrowserAliasDataSource/Source/ContentBrowserAliasDataSource/Private/ContentBrowserAliasDataSource.cpp b/Engine/Plugins/Editor/ContentBrowser/ContentBrowserAliasDataSource/Source/ContentBrowserAliasDataSource/Private/ContentBrowserAliasDataSource.cpp
new file mode 100644
--- /dev/null
+++ b/Engine/Plugins/Editor/ContentBrowser/ContentBrowserAliasDataSource/Source/ContentBrowserAliasDataSource/Private/ContentBrowserAliasDataSource.cpp
@@ -0,0 +1,1041 @@
+// Copyright Epic Games, Inc. All Rights Reserved.
+
+#include "ContentBrowserAliasDataSource.h"
+#include "ContentBrowserAliasDataSourceModule.h"
+
+#include "AssetRegistry/AssetRegistryModule.h"
+#include "ContentBrowserAssetDataCore.h"
+#include "ContentBrowserDataSubsystem.h"
+#include "ContentBrowserItemPath.h"
+
+#include "AssetToolsModule.h"
+#include "CollectionManagerModule.h"
+#include "IAssetTools.h"
+#include "IContentBrowserDataModule.h"
+#include "Misc/PackageName.h"
+
+#include UE_INLINE_GENERATED_CPP_BY_NAME(ContentBrowserAliasDataSource)
+
+DEFINE_LOG_CATEGORY(LogContentBrowserAliasDataSource);
+
+namespace ContentBrowserAliasDataSource
+{
+	FAutoConsoleCommand LogAliasesCommand = FAutoConsoleCommand(
+		TEXT("ContentBrowser.LogAliases"),
+		TEXT("Logs all content browser aliases"),
+		FConsoleCommandWithArgsDelegate::CreateLambda(
+			[](const TArray<FString>& /*Args*/)
+			{
+				if (UContentBrowserAliasDataSource* AliasDataSource = FModuleManager::Get().LoadModuleChecked<FContentBrowserAliasDataSourceModule>("ContentBrowserAliasDataSource").TryGetAliasDataSource())
+				{
+					AliasDataSource->LogAliases();
+				}
+			}));
+
+	template <typename AliasType>
+	FName GetAliasName(const AliasType&);
+
+	template <>
+	FName GetAliasName<FName>(const FName& Alias)
+	{
+		return Alias;
+	}
+
+	template <>
+	FName GetAliasName<FContentBrowserLocalizedAlias>(const FContentBrowserLocalizedAlias& Alias)
+	{
+		return Alias.Alias;
+	}
+
+	template <typename AliasType>
+	FText GetAliasDisplayNameOverride(const AliasType&);
+
+	template <>
+	FText GetAliasDisplayNameOverride<FName>(const FName& Alias)
+	{
+		return FText::GetEmpty();
+	}
+
+	template <>
+	FText GetAliasDisplayNameOverride<FContentBrowserLocalizedAlias>(const FContentBrowserLocalizedAlias& Alias)
+	{
+		return Alias.DisplayName;
+	}
+}
+
+FName UContentBrowserAliasDataSource::AliasTagName = "ContentBrowserAliases";
+
+void UContentBrowserAliasDataSource::Initialize(const bool InAutoRegister)
+{
+	Super::Initialize(InAutoRegister);
+	if (GIsEditor && !IsRunningCommandlet())
+	{
+		AssetRegistry = &FModuleManager::LoadModuleChecked<FAssetRegistryModule>(AssetRegistryConstants::ModuleName).Get();
+		AssetTools = &FModuleManager::GetModuleChecked<FAssetToolsModule>("AssetTools").Get();
+
+		AssetRegistry->OnAssetAdded().AddUObject(this, &UContentBrowserAliasDataSource::OnAssetAdded);
+		AssetRegistry->OnAssetRemoved().AddUObject(this, &UContentBrowserAliasDataSource::OnAssetRemoved);
+		AssetRegistry->OnAssetUpdated().AddUObject(this, &UContentBrowserAliasDataSource::OnAssetUpdated);
+		FCoreUObjectDelegates::OnAssetLoaded.AddUObject(this, &UContentBrowserAliasDataSource::OnAssetLoaded);
+		FCoreUObjectDelegates::OnObjectPropertyChanged.AddUObject(this, &UContentBrowserAliasDataSource::OnObjectPropertyChanged);
+	}
+}
+
+void UContentBrowserAliasDataSource::Shutdown()
+{
+	if (FModuleManager::Get().IsModuleLoaded(AssetRegistryConstants::ModuleName) && AssetRegistry)
+	{
+		AssetRegistry->OnAssetAdded().RemoveAll(this);
+		AssetRegistry->OnAssetRemoved().RemoveAll(this);
+		AssetRegistry->OnAssetUpdated().RemoveAll(this);
+	}
+
+	FCoreUObjectDelegates::OnAssetLoaded.RemoveAll(this);
+	FCoreUObjectDelegates::OnObjectPropertyChanged.RemoveAll(this);
+
+	Super::Shutdown();
+}
+
+void UContentBrowserAliasDataSource::BuildRootPathVirtualTree()
+{
+	Super::BuildRootPathVirtualTree();
+
+	PathTree.EnumerateSubPaths("/", [this](FName Path)
+	{
+		RootPathAdded(FNameBuilder(Path));
+		return true;
+	}, false);
+}
+
+void UContentBrowserAliasDataSource::CompileFilter(const FName InPath, const FContentBrowserDataFilter& InFilter, FContentBrowserDataCompiledFilter& OutCompiledFilter)
+{
+	TRACE_CPUPROFILER_EVENT_SCOPE(UContentBrowserAliasDataSource::CompileFilter);
+
+	UContentBrowserAssetDataSource::FAssetFilterInputParams Params;
+	if (UContentBrowserAssetDataSource::PopulateAssetFilterInputParams(Params, this, AssetRegistry, InFilter, OutCompiledFilter, &FCollectionManagerModule::GetModule().Get(), &FilterCache))
+	{
+		// Use the DataSource's custom PathTree instead of the AssetRegistry
+		const bool bCreatedPathFilter = UContentBrowserAssetDataSource::CreatePathFilter(Params, InPath, InFilter, OutCompiledFilter, [this](FName Path, TFunctionRef<bool(FName)> Callback, bool bRecursive)
+		{
+			PathTree.EnumerateSubPaths(Path, Callback, bRecursive);
+		});
+
+		if (bCreatedPathFilter)
+		{
+			auto CustomSubPathEnumeration = [this, &Params](FName Path, TFunctionRef<bool(FName)> Callback, bool bRecursive)
+				{
+					// Same as CreatePathFilter - CompileFilter calls EnumerateSubPaths internally so this needs to intercept
+					// the filter compilation and use its own PathTree to generate the sub paths.
+					AssetRegistry->EnumerateSubPaths(Path, Callback,bRecursive);
+					PathTree.EnumerateSubPaths(Path, Callback, bRecursive);
+				};
+
+			UContentBrowserAssetDataSource::FSubPathEnumerationFunc CustomSubPathEnumerationRef = CustomSubPathEnumeration;
+			UContentBrowserAssetDataSource::CreateAssetFilter(Params, InPath, InFilter, OutCompiledFilter, &CustomSubPathEnumerationRef);
+		}
+	}
+}
+
+void UContentBrowserAliasDataSource::EnumerateItemsMatchingFilter(const FContentBrowserDataCompiledFilter& InFilter, TFunctionRef<bool(FContentBrowserItemData&&)> InCallback)
+{
+	const FContentBrowserDataFilterList* FilterList = InFilter.CompiledFilters.Find(this);
+	if (!FilterList)
+	{
+		return;
+	}
+
+	const FContentBrowserCompiledAssetDataFilter* AssetDataFilter = FilterList->FindFilter<FContentBrowserCompiledAssetDataFilter>();
+	if (!AssetDataFilter)
+	{
+		return;
+	}
+
+	if (EnumHasAnyFlags(InFilter.ItemTypeFilter, EContentBrowserItemTypeFilter::IncludeFolders))
+	{
+		// Use the DataSource's custom PathTree instead of the AssetRegistry
+		auto EnumerateSubPaths = [this](FName Path, TFunctionRef<bool(FName)> Callback, bool bRecursive)
+		{
+			PathTree.EnumerateSubPaths(Path, Callback, bRecursive);
+		};
+		auto CreateFolderItem = [this](FName Path) -> FContentBrowserItemData
+		{
+			return CreateAssetFolderItem(Path);
+		};
+		UContentBrowserAssetDataSource::EnumerateFoldersMatchingFilter(this, AssetDataFilter, TGetOrEnumerateSink(InCallback), EnumerateSubPaths, CreateFolderItem);
+	}
+
+	if (EnumHasAnyFlags(InFilter.ItemTypeFilter, EContentBrowserItemTypeFilter::IncludeFiles) && !AssetDataFilter->bFilterExcludesAllAssets && !AssetDataFilter->InclusiveFilter.IsEmpty())
+	{
+		bool bIsAlreadyInSet = false;
+		if (AssetDataFilter->InclusiveFilter.PackagePaths.Num() > 0)
+		{
+			// Find all aliases for each requested package path and check if it passes both
+			// the inclusive and exclusive filters.
+			for (const FName PackagePath : AssetDataFilter->InclusiveFilter.PackagePaths)
+			{
+				const TArray<FContentBrowserUniqueAlias>* Aliases = AliasesInPackagePath.Find(PackagePath);
+				if (Aliases)
+				{
+					for (const FContentBrowserUniqueAlias& Alias : *Aliases)
+					{
+						const FAliasData& AliasData = AllAliases[Alias];
+						if (DoesAliasPassFilter(AliasData, *AssetDataFilter))
+						{
+							AlreadyAddedOriginalAssets.Add(AliasData.AssetData.GetSoftObjectPath(), &bIsAlreadyInSet);
+							if (!bIsAlreadyInSet)
+							{
+								if (!InCallback(CreateAssetFileItem(Alias)))
+								{
+									AlreadyAddedOriginalAssets.Reset();
+									return;
+								}
+							}
+						}
+					}
+				}
+			}
+		}
+		else
+		{
+			// If no package paths are requested, do the same as above for all aliases
+			for (const TPair<FContentBrowserUniqueAlias, FAliasData>& Pair : AllAliases)
+			{
+				if (DoesAliasPassFilter(Pair.Value, *AssetDataFilter))
+				{
+					AlreadyAddedOriginalAssets.Add(Pair.Value.AssetData.GetSoftObjectPath(), &bIsAlreadyInSet);
+					if (!bIsAlreadyInSet)
+					{
+						if (!InCallback(CreateAssetFileItem(Pair.Key)))
+						{
+							AlreadyAddedOriginalAssets.Reset();
+							return;
+						}
+					}
+				}
+			}
+		}
+		AlreadyAddedOriginalAssets.Reset();
+	}
+}
+
+void UContentBrowserAliasDataSource::EnumerateItemsAtPath(const FName InPath, const EContentBrowserItemTypeFilter InItemTypeFilter, TFunctionRef<bool(FContentBrowserItemData&&)> InCallback)
+{
+	FName InternalPath;
+	if (!TryConvertVirtualPathToInternal(InPath, InternalPath))
+	{
+		return;
+	}
+
+	if (EnumHasAnyFlags(InItemTypeFilter, EContentBrowserItemTypeFilter::IncludeFolders))
+	{
+		if (PathTree.PathExists(InternalPath))
+		{
+			InCallback(CreateAssetFolderItem(InternalPath));
+		}
+	}
+
+	if (EnumHasAnyFlags(InItemTypeFilter, EContentBrowserItemTypeFilter::IncludeFiles))
+	{
+		TSet<FName> Paths;
+		// Return all assets for this path and its subpaths
+		Paths.Add(InternalPath);
+		PathTree.GetSubPaths(InternalPath, Paths, true);
+
+		bool bIsAlreadyInSet = false;
+		for (const FName Path : Paths)
+		{
+			const TArray<FContentBrowserUniqueAlias>* Aliases = AliasesInPackagePath.Find(Path);
+			if (Aliases)
+			{
+				for (const FContentBrowserUniqueAlias& Alias : *Aliases)
+				{
+					AlreadyAddedOriginalAssets.Add(AllAliases[Alias].AssetData.GetSoftObjectPath(), &bIsAlreadyInSet);
+					if (!bIsAlreadyInSet)
+					{
+						if (!InCallback(CreateAssetFileItem(Alias)))
+						{
+							AlreadyAddedOriginalAssets.Reset();
+							return;
+						}
+					}
+				}
+			}
+		}
+		AlreadyAddedOriginalAssets.Reset();
+	}
+}
+
+bool UContentBrowserAliasDataSource::DoesItemPassFilter(const FContentBrowserItemData& InItem, const FContentBrowserDataCompiledFilter& InFilter)
+{
+	const FContentBrowserDataFilterList* FilterList = InFilter.CompiledFilters.Find(this);
+	if (!FilterList)
+	{
+		return false;
+	}
+
+	const FContentBrowserCompiledAssetDataFilter* AssetDataFilter = FilterList->FindFilter<FContentBrowserCompiledAssetDataFilter>();
+	if (!AssetDataFilter)
+	{
+		return false;
+	}
+
+	switch (InItem.GetItemType())
+	{
+	case EContentBrowserItemFlags::Type_Folder:
+		if (EnumHasAnyFlags(InFilter.ItemTypeFilter, EContentBrowserItemTypeFilter::IncludeFolders))
+		{
+			return UContentBrowserAssetDataSource::DoesItemPassFolderFilter(this, InItem, *AssetDataFilter);
+		}
+		break;
+
+	case EContentBrowserItemFlags::Type_File:
+		if (EnumHasAnyFlags(InFilter.ItemTypeFilter, EContentBrowserItemTypeFilter::IncludeFiles) && !AssetDataFilter->bFilterExcludesAllAssets)
+		{
+			if (TSharedPtr<const FContentBrowserAliasItemDataPayload> AliasPayload = StaticCastSharedPtr<const FContentBrowserAliasItemDataPayload>(InItem.GetPayload()))
+			{
+				FAliasData* FoundAlias = AllAliases.Find(AliasPayload->Alias);
+				// This should always be true except in the case where the item is deleted in the same tick that DoesItemPassFilter is called.
+				// This is because the alias data source processes the deletion immediately, but the content browser deletion is queued until
+				// next tick, causing them to be briefly out of sync.
+				// An alternative solution would be to add a way to synchronously flush content browser updates, but that doesn't exist atm.
+				if (FoundAlias)
+				{
+					return DoesAliasPassFilter(*FoundAlias, *AssetDataFilter);
+				}
+			}
+		}
+		break;
+	}
+
+	return false;
+}
+
+bool UContentBrowserAliasDataSource::DoesAliasPassFilter(const FAliasData& AliasData, const FContentBrowserCompiledAssetDataFilter& Filter) const
+{
+	// Create a fake asset data using the alias path instead of the asset's original path
+	// AssetRegistry->IsAssetIncludedByFilter is effectively a static function and does not actually use AssetRegistry data
+	FAssetData AliasAssetData(AliasData.PackageName, AliasData.PackagePath, AliasData.AssetData.AssetName, AliasData.AssetData.AssetClassPath, AliasData.AssetData.TagsAndValues.CopyMap());
+
+	return (Filter.InclusiveFilter.IsEmpty() || AssetRegistry->IsAssetIncludedByFilter(AliasAssetData, Filter.InclusiveFilter)) // Passes Inclusive
+		&& (Filter.ExclusiveFilter.IsEmpty() || !AssetRegistry->IsAssetIncludedByFilter(AliasAssetData, Filter.ExclusiveFilter)); // Passes Exclusive
+}
+
+TArray<FContentBrowserItemPath> UContentBrowserAliasDataSource::GetAliasesForPath(const FSoftObjectPath& ObjectPath) const
+{
+	TArray<FContentBrowserItemPath> OutAliases;
+
+	if (const TArray<FName>* FoundAliases = AliasesForObjectPath.Find(ObjectPath))
+	{
+		for (const FName Alias : *FoundAliases)
+		{
+			if (const FAliasData* AliasData = AllAliases.Find(FContentBrowserUniqueAlias(ObjectPath, Alias)))
+			{
+				OutAliases.Add(FContentBrowserItemPath(AliasData->PackageName, EContentBrowserPathType::Internal));
+			}
+		}
+	}
+
+	return OutAliases;
+}
+
+void UContentBrowserAliasDataSource::AddAliases(const FAssetData& Asset, const TArray<FName>& Aliases, const bool bInIsFromMetaData, const bool bSkipPrimaryAssetValidation)
+{
+	AddAliasesImpl(Asset, Aliases, bInIsFromMetaData, bSkipPrimaryAssetValidation);
+}
+
+void UContentBrowserAliasDataSource::AddAliases(const FAssetData& Asset, const TArray<FContentBrowserLocalizedAlias>& Aliases, const bool bInIsFromMetaData, const bool bSkipPrimaryAssetValidation)
+{
+	AddAliasesImpl(Asset, Aliases, bInIsFromMetaData, bSkipPrimaryAssetValidation);
+}
+
+template <typename AliasType>
+void UContentBrowserAliasDataSource::AddAliasesImpl(const FAssetData& Asset, const TArray<AliasType>& Aliases, const bool bInIsFromMetaData, const bool bSkipPrimaryAssetValidation)
+{
+	for (const AliasType& Alias : Aliases)
+	{
+		AddAlias(Asset, Alias, bInIsFromMetaData, bSkipPrimaryAssetValidation);
+	}
+}
+
+void UContentBrowserAliasDataSource::AddAlias(const FAssetData& Asset, const FName Alias, const bool bInIsFromMetaData, const bool bSkipPrimaryAssetValidation)
+{
+	AddAliasImpl(Asset, Alias, bInIsFromMetaData, bSkipPrimaryAssetValidation);
+}
+
+void UContentBrowserAliasDataSource::AddAlias(const FAssetData& Asset, const FContentBrowserLocalizedAlias& Alias, const bool bInIsFromMetaData, const bool bSkipPrimaryAssetValidation)
+{
+	AddAliasImpl(Asset, Alias, bInIsFromMetaData, bSkipPrimaryAssetValidation);
+}
+
+template <typename AliasType>
+void UContentBrowserAliasDataSource::AddAliasImpl(const FAssetData& Asset, const AliasType& AliasToAdd, const bool bInIsFromMetaData, const bool bSkipPrimaryAssetValidation)
+{
+	const FName Alias = ContentBrowserAliasDataSource::GetAliasName(AliasToAdd);
+	const FText AliasDisplayNameOverride = ContentBrowserAliasDataSource::GetAliasDisplayNameOverride(AliasToAdd);
+
+	auto LogErrorMessage = [&Alias, &Asset](const TCHAR* Reason)
+	{
+		UE_LOG(LogContentBrowserAliasDataSource, Warning, TEXT("Cannot add alias %s for %s because: %s"), *Alias.ToString(), *Asset.GetObjectPathString(), Reason);
+	};
+	
+	FContentBrowserUniqueAlias UniqueAlias = MakeTuple(Asset.GetSoftObjectPath(), Alias);
+	if (AllAliases.Contains(UniqueAlias))
+	{
+		LogErrorMessage(TEXT("An alias with that name already exists"));
+		return;
+	}
+
+	if (!bSkipPrimaryAssetValidation && !ContentBrowserAssetData::IsPrimaryAsset(Asset))
+	{
+		LogErrorMessage(TEXT("It is not a primary asset"));
+		return;
+	}
+
+	const FString AliasString = Alias.ToString();
+	if (AliasString[0] != TEXT('/'))
+	{
+		LogErrorMessage(TEXT("The alias is not a valid path"));
+		return;
+	}
+
+	TArray<FString> Tokens;
+	// TODO: figure out how to preserve spaces but also do some kind of invalid character checking that works even with fake object paths
+	//Alias = ObjectTools::SanitizeObjectPath(Alias);
+	AliasString.ParseIntoArray(Tokens, TEXT("/"), false);
+	// Minimum valid tokens = [Empty, PackageName, ObjectName]
+	if (Tokens.Num() < 3)
+	{
+		LogErrorMessage(TEXT("The alias is not a valid path"));
+		return;
+	}
+
+	// Check for invalid empty tokens. Skip token 0 since it will be empty due to / at start
+	for (int32 i = 1; i < Tokens.Num(); ++i)
+	{
+		if (Tokens[i].Len() == 0)
+		{
+			LogErrorMessage(TEXT("The alias is not a valid path"));
+			return;
+		}
+	}
+	FStringView AliasStringView(AliasString);
+	// Tokens[1] is the root path without the slash prefix
+	RootPathAdded(AliasStringView.Left(Tokens[1].Len() + 1));
+
+	// PackagePath is everything before the last slash
+	const FName PackagePath(AliasString.Len() - Tokens.Last().Len() - 1, AliasStringView.GetData());
+	PathTree.CachePath(PackagePath, [this](FName AddedPath)
+	{
+		QueueItemDataUpdate(FContentBrowserItemDataUpdate::MakeItemAddedUpdate(CreateAssetFolderItem(AddedPath)));
+	});
+
+	AliasesInPackagePath.FindOrAdd(PackagePath).Add(UniqueAlias);
+	AliasesForObjectPath.FindOrAdd(Asset.GetSoftObjectPath()).Add(Alias);
+	AllAliases.Add(UniqueAlias, FAliasData(Asset, PackagePath, AliasDisplayNameOverride.IsEmpty() ? FText::AsCultureInvariant(*Tokens.Last()) : AliasDisplayNameOverride, bInIsFromMetaData));
+	QueueItemDataUpdate(FContentBrowserItemDataUpdate::MakeItemAddedUpdate(CreateAssetFileItem(UniqueAlias)));
+
+	// This logging might get out of control if there ends up being hundreds of thousands of aliases.
+	//UE_LOG(LogContentBrowserAliasDataSource, VeryVerbose, TEXT("Adding alias %s for %s"), *AliasString, *Asset.ObjectPath.ToString());
+}
+
+void UContentBrowserAliasDataSource::RemoveFoldersRecursively(FStringView LeafFolder)
+{
+	// Make sure there's no assets here
+	const FName LeafFolderName(LeafFolder.Len(), LeafFolder.GetData());
+	if (LeafFolderName != NAME_None && !AliasesInPackagePath.Contains(LeafFolderName))
+	{
+		// Make sure there's no other child folders
+		bool bHasChildren = false;
+		PathTree.EnumerateSubPaths(LeafFolderName, [&bHasChildren](FName SubPath)
+		{
+			bHasChildren = true;
+			return false;
+		}, false);
+		if (!bHasChildren)
+		{
+			// Remove folder from PathTree
+			PathTree.RemovePath(LeafFolderName, [](FName Path) {});
+			QueueItemDataUpdate(FContentBrowserItemDataUpdate::MakeItemRemovedUpdate(CreateAssetFolderItem(LeafFolderName)));
+
+			// Check parent folder
+			int32 LastSlash = INDEX_NONE;
+			LeafFolder.FindLastChar(TEXT('/'), LastSlash);
+			FStringView NewLeafFolder = LeafFolder.Left(LastSlash);
+
+			// If the last slash is at the start of the string, this is the root folder
+			if (LastSlash == 1)
+			{
+				RootPathRemoved(NewLeafFolder);
+			}
+			else
+			{
+				RemoveFoldersRecursively(NewLeafFolder);
+			}
+		}
+	}
+}
+
+PRAGMA_DISABLE_DEPRECATION_WARNINGS
+void UContentBrowserAliasDataSource::RemoveAlias(FName ObjectPath, const FName Alias)
+{
+	RemoveAlias(FSoftObjectPath(ObjectPath), Alias);
+}
+PRAGMA_ENABLE_DEPRECATION_WARNINGS
+
+void UContentBrowserAliasDataSource::RemoveAlias(const FSoftObjectPath& ObjectPath, const FName Alias)
+{
+    FContentBrowserUniqueAlias UniqueAlias = MakeTuple(ObjectPath, Alias);
+	FAliasData AliasData;
+	// Store a copy of the item data before it's removed for the MakeItemRemovedUpdate notification
+	FContentBrowserItemData ItemData = CreateAssetFileItem(UniqueAlias);
+	if (AllAliases.RemoveAndCopyValue(UniqueAlias, AliasData))
+	{
+		check(AliasData.AssetData.GetSoftObjectPath() == ObjectPath);
+		{
+			TArray<FName>& Aliases = AliasesForObjectPath[ObjectPath];
+			Aliases.Remove(Alias);
+			if (Aliases.IsEmpty())
+			{
+				AliasesForObjectPath.Remove(ObjectPath);
+			}
+		}
+
+		{
+			TArray<FContentBrowserUniqueAlias>& Aliases = AliasesInPackagePath[AliasData.PackagePath];
+			Aliases.Remove(UniqueAlias);
+			if (Aliases.IsEmpty())
+			{
+				AliasesInPackagePath.Remove(AliasData.PackagePath);
+			}
+		}
+
+		QueueItemDataUpdate(FContentBrowserItemDataUpdate::MakeItemRemovedUpdate(ItemData));
+
+		FNameBuilder AliasAsString(Alias);
+		FStringView AliasStringView(AliasAsString);
+
+		int32 LastSlash = INDEX_NONE;
+		AliasStringView.FindLastChar(TEXT('/'), LastSlash);
+		RemoveFoldersRecursively(AliasStringView.Left(LastSlash));
+	}
+}
+
+void UContentBrowserAliasDataSource::AddAliasFolderDisplayName(const FName AliasFolder, const FText& DisplayName)
+{
+	if (const FText* ExistingDisplayName = AliasFolderDisplayNames.Find(AliasFolder);
+		ExistingDisplayName && ExistingDisplayName->IdenticalTo(DisplayName, ETextIdenticalModeFlags::LexicalCompareInvariants))
+	{
+		return;
+	}
+
+	AliasFolderDisplayNames.Add(AliasFolder, DisplayName);
+
+	if (PathTree.PathExists(AliasFolder))
+	{
+		QueueItemDataUpdate(FContentBrowserItemDataUpdate::MakeItemModifiedUpdate(CreateAssetFolderItem(AliasFolder)));
+	}
+}
+
+void UContentBrowserAliasDataSource::RemoveAliasFolderDisplayName(const FName AliasFolder)
+{
+	AliasFolderDisplayNames.Remove(AliasFolder);
+
+	if (PathTree.PathExists(AliasFolder))
+	{
+		QueueItemDataUpdate(FContentBrowserItemDataUpdate::MakeItemModifiedUpdate(CreateAssetFolderItem(AliasFolder)));
+	}
+}
+
+void UContentBrowserAliasDataSource::RebuildAliases()
+{
+	PathTree = FPathTree();
+	AllAliases.Reset();
+	AliasesForObjectPath.Reset();
+	AliasesInPackagePath.Reset();
+	AliasFolderDisplayNames.Reset();
+	AlreadyAddedOriginalAssets.Reset();
+	FilterCache.Reset();
+	RootPathVirtualTree.Reset();
+
+	OnRebuildAliases().Broadcast();
+}
+
+PRAGMA_DISABLE_DEPRECATION_WARNINGS
+void UContentBrowserAliasDataSource::RemoveAliases(FName ObjectPath)
+{
+	RemoveAliases(FSoftObjectPath(ObjectPath));
+}
+PRAGMA_ENABLE_DEPRECATION_WARNINGS
+
+void UContentBrowserAliasDataSource::RemoveAliases(const FSoftObjectPath& ObjectPath)
+{
+	const TArray<FName>* AliasesPtr = AliasesForObjectPath.Find(ObjectPath);
+	if (AliasesPtr)
+	{
+		// Create a copy to not modify array during iteration
+		const TArray<FName> Aliases = *AliasesPtr;
+		for (const FName Alias : Aliases)
+		{
+			RemoveAlias(ObjectPath, Alias);
+		}
+	}
+}
+
+void UContentBrowserAliasDataSource::OnAssetAdded(const FAssetData& InAssetData)
+{
+	FString AliasTagValue;
+	if (InAssetData.GetTagValue(AliasTagName, AliasTagValue))
+	{
+		TArray<FString> AliasList;
+		AliasTagValue.ParseIntoArray(AliasList, TEXT(","));
+		for (FString Alias : AliasList)
+		{
+			AddAlias(InAssetData, *Alias, true);
+		}
+	}
+}
+
+void UContentBrowserAliasDataSource::OnAssetRemoved(const FAssetData& InAssetData)
+{
+	RemoveAliases(InAssetData);
+}
+
+void UContentBrowserAliasDataSource::ReconcileAliasesFromMetaData(const FAssetData& Asset)
+{
+	const TArray<FName>* ExistingAliasesPtr = AliasesForObjectPath.Find(Asset.GetSoftObjectPath());
+	if (ExistingAliasesPtr)
+	{
+		FContentBrowserUniqueAlias UniqueAlias = MakeTuple(Asset.GetSoftObjectPath(), NAME_None);
+		
+		FString AliasTagValue;
+		if (Asset.GetTagValue(AliasTagName, AliasTagValue))
+		{
+			// Reconcile existing aliases from metadata vs new aliases from metadata
+			TArray<FString> AliasesFromTag;
+			AliasTagValue.ParseIntoArray(AliasesFromTag, TEXT(","));
+
+			TArray<FName> AliasNamesFromTag;
+			AliasNamesFromTag.Reserve(AliasesFromTag.Num());
+			for (const FString& AliasFromTag : AliasesFromTag)
+			{
+				AliasNamesFromTag.Add(*AliasFromTag);
+			}
+
+			TArray<FName> AliasesOnlyInExisting, AliasesOnlyInNew;
+			for (const FName Alias : *ExistingAliasesPtr)
+			{
+				UniqueAlias.Value = Alias;
+				if (AllAliases[UniqueAlias].bIsFromMetaData)
+				{
+					if (!AliasNamesFromTag.Contains(Alias))
+					{
+						AliasesOnlyInExisting.Add(Alias);
+					}
+				}
+			}
+
+			for (const FName Alias : AliasNamesFromTag)
+			{
+				if (!ExistingAliasesPtr->Contains(Alias))
+				{
+					AliasesOnlyInNew.Add(Alias);
+				}
+			}
+
+			for (const FName ExistingAlias : AliasesOnlyInExisting)
+			{
+				RemoveAlias(Asset.GetSoftObjectPath(), ExistingAlias);
+			}
+			AddAliases(Asset, AliasesOnlyInNew, true);
+		}
+		else
+		{
+			// If the tag was removed, then remove any existing aliases generated from metadata
+			TArray<FName> ExistingAliases = *ExistingAliasesPtr;
+			for (const FName Alias : ExistingAliases)
+			{
+				UniqueAlias.Value = Alias;
+				if (AllAliases[UniqueAlias].bIsFromMetaData)
+				{
+					RemoveAlias(Asset.GetSoftObjectPath(), Alias);
+				}
+			}
+		}
+	}
+	else
+	{
+		// If no existing aliases found, then check if new metadata was added
+		OnAssetAdded(Asset);
+	}
+}
+
+void UContentBrowserAliasDataSource::ReconcileAliasesForAsset(const FAssetData& Asset, const TArray<FName>& NewAliases)
+{
+	ReconcileAliasesForAssetImpl(Asset, NewAliases);
+}
+
+void UContentBrowserAliasDataSource::ReconcileAliasesForAsset(const FAssetData& Asset, const TArray<FContentBrowserLocalizedAlias>& NewAliases)
+{
+	ReconcileAliasesForAssetImpl(Asset, NewAliases);
+}
+
+template <typename AliasType>
+void UContentBrowserAliasDataSource::ReconcileAliasesForAssetImpl(const FAssetData& Asset, const TArray<AliasType>& NewAliases)
+{
+	if (ensure(Asset.IsValid()))
+	{
+		const TArray<FName>* ExistingAliasesPtr = AliasesForObjectPath.Find(Asset.GetSoftObjectPath());
+		if (ExistingAliasesPtr)
+		{
+			TArray<FName> AliasesOnlyInExisting;
+			TArray<AliasType> AliasesOnlyInNew;
+			FContentBrowserUniqueAlias UniqueAlias = MakeTuple(Asset.GetSoftObjectPath(), NAME_None);
+			for (const FName Alias : *ExistingAliasesPtr)
+			{
+				UniqueAlias.Value = Alias;
+				if (!AllAliases[UniqueAlias].bIsFromMetaData)
+				{
+					if (!NewAliases.ContainsByPredicate([&Alias](const AliasType& NewAlias) { return ContentBrowserAliasDataSource::GetAliasName(NewAlias) == Alias; }))
+					{
+						AliasesOnlyInExisting.Add(Alias);
+					}
+				}
+			}
+
+			for (const AliasType& Alias : NewAliases)
+			{
+				if (!ExistingAliasesPtr->Contains(ContentBrowserAliasDataSource::GetAliasName(Alias)))
+				{
+					AliasesOnlyInNew.Add(Alias);
+				}
+			}
+
+			for (const FName ExistingAlias : AliasesOnlyInExisting)
+			{
+				RemoveAlias(Asset.GetSoftObjectPath(), ExistingAlias);
+			}
+			AddAliases(Asset, AliasesOnlyInNew);
+		}
+		else
+		{
+			AddAliases(Asset, NewAliases);
+		}
+	}
+}
+
+void UContentBrowserAliasDataSource::LogAliases() const
+{
+	TArray<FContentBrowserUniqueAlias> Aliases;
+	AllAliases.GenerateKeyArray(Aliases);
+	Aliases.Sort(
+		[](const FContentBrowserUniqueAlias& A, const FContentBrowserUniqueAlias& B)
+		{
+			return A.Key == B.Key ? A.Value.LexicalLess(B.Value) : A.Key.LexicalLess(B.Key);
+		});
+
+	TStringBuilder<FName::StringBufferSize * 2> StringBuilder;
+	for (const FContentBrowserUniqueAlias& Alias : Aliases)
+	{
+		StringBuilder.Reset();
+		StringBuilder.Append(Alias.Key.ToString());
+		StringBuilder.Append(TEXT(" -> \""));
+		StringBuilder.Append(Alias.Value.ToString());
+		StringBuilder.AppendChar(TCHAR('"'));
+		UE_LOG(LogContentBrowserAliasDataSource, Log, TEXT("%s"), StringBuilder.ToString());
+	}
+}
+
+void UContentBrowserAliasDataSource::OnAssetUpdated(const FAssetData& InAssetData)
+{
+	ReconcileAliasesFromMetaData(InAssetData);
+	UpdateAliasesCachedAssetData(InAssetData);
+	MakeItemModifiedUpdate(InAssetData.GetSoftObjectPath());
+}
+
+void UContentBrowserAliasDataSource::UpdateAliasesCachedAssetData(const FAssetData& InAssetData)
+{
+	const FSoftObjectPath ObjectPath = InAssetData.GetSoftObjectPath();
+	if (const TArray<FName>* Aliases = AliasesForObjectPath.Find(ObjectPath))
+	{
+		FContentBrowserUniqueAlias UniqueAlias = MakeTuple(ObjectPath, NAME_None);
+		for (const FName Alias : *Aliases)
+		{
+			UniqueAlias.Value = Alias;
+			if (FAliasData* AliasData = AllAliases.Find(UniqueAlias))
+			{
+				AliasData->AssetData = InAssetData;
+			}
+		}
+	}
+}
+
+void UContentBrowserAliasDataSource::MakeItemModifiedUpdate(const FSoftObjectPath& ObjectPath)
+{
+	if (const TArray<FName>* Aliases = AliasesForObjectPath.Find(ObjectPath))
+	{
+		FContentBrowserUniqueAlias UniqueAlias = MakeTuple(ObjectPath, NAME_None);
+		for (const FName Alias : *Aliases)
+		{
+			UniqueAlias.Value = Alias;
+			QueueItemDataUpdate(FContentBrowserItemDataUpdate::MakeItemModifiedUpdate(CreateAssetFileItem(UniqueAlias)));
+		}
+	}
+}
+
+void UContentBrowserAliasDataSource::OnAssetLoaded(UObject* InAsset)
+{
+	if (InAsset && !InAsset->GetOutermost()->HasAnyPackageFlags(PKG_ForDiffing))
+	{
+		const FAssetData LoadedAssetData(InAsset);
+		ReconcileAliasesFromMetaData(LoadedAssetData);
+		UpdateAliasesCachedAssetData(LoadedAssetData);
+		MakeItemModifiedUpdate(InAsset);
+	}
+}
+
+void UContentBrowserAliasDataSource::OnObjectPropertyChanged(UObject* InObject, FPropertyChangedEvent& InPropertyChangedEvent)
+{
+	if (InObject && InObject->IsAsset() && AliasesForObjectPath.Contains(InObject))
+	{
+		const FAssetData LoadedAssetData(InObject);
+		UpdateAliasesCachedAssetData(LoadedAssetData);
+		MakeItemModifiedUpdate(InObject);
+	}
+}
+
+bool UContentBrowserAliasDataSource::GetItemAttribute(const FContentBrowserItemData& InItem, const bool InIncludeMetaData, const FName InAttributeKey, FContentBrowserItemDataAttributeValue& OutAttributeValue)
+{
+	return ContentBrowserAssetData::GetItemAttribute(this, InItem, InIncludeMetaData, InAttributeKey, OutAttributeValue);
+}
+
+bool UContentBrowserAliasDataSource::GetItemAttributes(const FContentBrowserItemData& InItem, const bool InIncludeMetaData, FContentBrowserItemDataAttributeValues& OutAttributeValues)
+{
+	return ContentBrowserAssetData::GetItemAttributes(this, InItem, InIncludeMetaData, OutAttributeValues);
+}
+
+bool UContentBrowserAliasDataSource::GetItemPhysicalPath(const FContentBrowserItemData& InItem, FString& OutDiskPath)
+{
+	if (InItem.GetItemType() == EContentBrowserItemFlags::Category_Asset)
+	{
+		return ContentBrowserAssetData::GetItemPhysicalPath(this, InItem, OutDiskPath);
+	}
+	return false;
+}
+
+bool UContentBrowserAliasDataSource::IsItemDirty(const FContentBrowserItemData& InItem)
+{
+	return ContentBrowserAssetData::IsItemDirty(this, InItem);
+}
+
+bool UContentBrowserAliasDataSource::CanEditItem(const FContentBrowserItemData& InItem, FText* OutErrorMsg)
+{
+	if (TSharedPtr<const FContentBrowserAliasItemDataPayload> AliasPayload = StaticCastSharedPtr<const FContentBrowserAliasItemDataPayload>(InItem.GetPayload()))
+	{
+		// Both the alias path and asset path must pass the writable folder filter and editable folder filter in order to be editable
+		const TSharedRef<FPathPermissionList>& WritableFolderFilter = AssetTools->GetWritableFolderPermissionList();
+		if (!WritableFolderFilter->PassesStartsWithFilter(AliasPayload->Alias.Value))
+		{
+			if (OutErrorMsg)
+			{
+				*OutErrorMsg = FText::Format(NSLOCTEXT("ContentBrowserAliasDataSource", "Error_FolderIsLocked", "Alias asset '{0}' is in a read only folder. Unable to edit read only assets."), FText::FromName(AliasPayload->Alias.Value));
+			}
+			return false;
+		}
+
+		if (UContentBrowserDataSubsystem* ContentBrowserDataSubsystem = IContentBrowserDataModule::Get().GetSubsystem())
+		{
+			const TSharedRef<FPathPermissionList>& EditableFolderFilter = ContentBrowserDataSubsystem->GetEditableFolderPermissionList();
+			if (!EditableFolderFilter->PassesStartsWithFilter(AliasPayload->Alias.Value))
+			{
+				if (OutErrorMsg)
+				{
+					*OutErrorMsg = FText::Format(NSLOCTEXT("ContentBrowserAliasDataSource", "Error_FolderIsNotEditable", "Alias asset '{0}' is in a folder that does not allow edits. Unable to edit read only assets."), FText::FromName(AliasPayload->Alias.Value));
+				}
+				return false;
+			}
+		}
+	}
+	return ContentBrowserAssetData::CanEditItem(AssetTools, this, InItem, OutErrorMsg);
+}
+
+bool UContentBrowserAliasDataSource::EditItem(const FContentBrowserItemData& InItem)
+{
+	return ContentBrowserAssetData::EditItems(AssetTools, this, MakeArrayView(&InItem, 1));
+}
+
+bool UContentBrowserAliasDataSource::BulkEditItems(TArrayView<const FContentBrowserItemData> InItems)
+{
+	return ContentBrowserAssetData::EditItems(AssetTools, this, InItems);
+}
+
+bool UContentBrowserAliasDataSource::CanViewItem(const FContentBrowserItemData& InItem, FText* OutErrorMsg)
+{
+	return ContentBrowserAssetData::CanViewItem(AssetTools, this, InItem, OutErrorMsg);
+}
+
+bool UContentBrowserAliasDataSource::ViewItem(const FContentBrowserItemData& InItem)
+{
+	return ContentBrowserAssetData::ViewItems(AssetTools, this, MakeArrayView(&InItem, 1));
+}
+
+bool UContentBrowserAliasDataSource::BulkViewItems(TArrayView<const FContentBrowserItemData> InItems)
+{
+	return ContentBrowserAssetData::ViewItems(AssetTools, this, InItems);
+}
+
+bool UContentBrowserAliasDataSource::CanPreviewItem(const FContentBrowserItemData& InItem, FText* OutErrorMsg)
+{
+	return ContentBrowserAssetData::CanPreviewItem(AssetTools, this, InItem, OutErrorMsg);
+}
+
+bool UContentBrowserAliasDataSource::PreviewItem(const FContentBrowserItemData& InItem)
+{
+	return ContentBrowserAssetData::PreviewItems(AssetTools, this, MakeArrayView(&InItem, 1));
+}
+
+bool UContentBrowserAliasDataSource::BulkPreviewItems(TArrayView<const FContentBrowserItemData> InItems)
+{
+	return ContentBrowserAssetData::PreviewItems(AssetTools, this, InItems);
+}
+
+bool UContentBrowserAliasDataSource::CanSaveItem(const FContentBrowserItemData& InItem, const EContentBrowserItemSaveFlags InSaveFlags, FText* OutErrorMsg)
+{
+	return ContentBrowserAssetData::CanSaveItem(AssetTools, this, InItem, InSaveFlags, OutErrorMsg);
+}
+
+bool UContentBrowserAliasDataSource::SaveItem(const FContentBrowserItemData& InItem, const EContentBrowserItemSaveFlags InSaveFlags)
+{
+	return ContentBrowserAssetData::SaveItems(AssetTools, this, MakeArrayView(&InItem, 1), InSaveFlags);
+}
+
+bool UContentBrowserAliasDataSource::BulkSaveItems(TArrayView<const FContentBrowserItemData> InItems, const EContentBrowserItemSaveFlags InSaveFlags)
+{
+	return ContentBrowserAssetData::SaveItems(AssetTools, this, InItems, InSaveFlags);
+}
+
+bool UContentBrowserAliasDataSource::CanPrivatizeItem(const FContentBrowserItemData& InItem, FText* OutErrorMsg)
+{
+	return ContentBrowserAssetData::CanPrivatizeItem(AssetTools, AssetRegistry, this, InItem, OutErrorMsg);
+}
+
+bool UContentBrowserAliasDataSource::PrivatizeItem(const FContentBrowserItemData& InItem)
+{
+	return ContentBrowserAssetData::PrivatizeItems(AssetTools, AssetRegistry, this, MakeArrayView(&InItem, 1));
+}
+
+bool UContentBrowserAliasDataSource::BulkPrivatizeItems(TArrayView<const FContentBrowserItemData> InItems)
+{
+	return ContentBrowserAssetData::PrivatizeItems(AssetTools, AssetRegistry, this, InItems);
+}
+
+bool UContentBrowserAliasDataSource::AppendItemReference(const FContentBrowserItemData& InItem, FString& InOutStr)
+{
+	return ContentBrowserAssetData::AppendItemReference(AssetRegistry, this, InItem, InOutStr);
+}
+
+bool UContentBrowserAliasDataSource::AppendItemObjectPath(const FContentBrowserItemData& InItem, FString& InOutStr)
+{
+	return ContentBrowserAssetData::AppendItemObjectPath(AssetRegistry, this, InItem, InOutStr);
+}
+
+bool UContentBrowserAliasDataSource::AppendItemPackageName(const FContentBrowserItemData& InItem, FString& InOutStr)
+{
+	return ContentBrowserAssetData::AppendItemPackageName(AssetRegistry, this, InItem, InOutStr);
+}
+
+bool UContentBrowserAliasDataSource::UpdateThumbnail(const FContentBrowserItemData& InItem, FAssetThumbnail& InThumbnail)
+{
+	return ContentBrowserAssetData::UpdateItemThumbnail(this, InItem, InThumbnail);
+}
+
+bool UContentBrowserAliasDataSource::TryGetCollectionId(const FContentBrowserItemData& InItem, FSoftObjectPath& OutCollectionId)
+{
+	return false;
+}
+
+bool UContentBrowserAliasDataSource::Legacy_TryGetPackagePath(const FContentBrowserItemData& InItem, FName& OutPackagePath)
+{
+	if (TSharedPtr<const FContentBrowserAssetFolderItemDataPayload> FolderPayload = ContentBrowserAssetData::GetAssetFolderItemPayload(this, InItem))
+	{
+		OutPackagePath = FolderPayload->GetInternalPath();
+		return true;
+	}
+	return false;
+}
+
+bool UContentBrowserAliasDataSource::Legacy_TryGetAssetData(const FContentBrowserItemData& InItem, FAssetData& OutAssetData)
+{
+	if (TSharedPtr<const FContentBrowserAssetFileItemDataPayload> AssetPayload = ContentBrowserAssetData::GetAssetFileItemPayload(this, InItem))
+	{
+		OutAssetData = AssetPayload->GetAssetData();
+		return true;
+	}
+	return false;
+}
+
+bool UContentBrowserAliasDataSource::Legacy_TryConvertPackagePathToVirtualPath(const FName InPackagePath, FName& OutPath)
+{
+	return PathTree.PathExists(InPackagePath) // Ignore unknown content paths
+		&& TryConvertInternalPathToVirtual(InPackagePath, OutPath);
+}
+
+bool UContentBrowserAliasDataSource::Legacy_TryConvertAssetDataToVirtualPath(const FAssetData& InAssetData, const bool InUseFolderPaths, FName& OutPath)
+{
+	return InAssetData.AssetClassPath != FTopLevelAssetPath(TEXT("/Script/CoreUObject"), NAME_Class) // Ignore legacy class items
+PRAGMA_DISABLE_DEPRECATION_WARNINGS
+		&& TryConvertInternalPathToVirtual(InUseFolderPaths ? InAssetData.PackagePath : InAssetData.ObjectPath, OutPath);
+PRAGMA_ENABLE_DEPRECATION_WARNINGS
+}
+
+void UContentBrowserAliasDataSource::RemoveUnusedCachedFilterData(const FContentBrowserDataFilterCacheIDOwner& IDOwner, TArrayView<const FName> InVirtualPathsInUse, const FContentBrowserDataFilter& DataFilter)
+{
+	FilterCache.RemoveUnusedCachedData(IDOwner, InVirtualPathsInUse, DataFilter);
+}
+
+void UContentBrowserAliasDataSource::ClearCachedFilterData(const FContentBrowserDataFilterCacheIDOwner& IDOwner)
+{
+	FilterCache.ClearCachedData(IDOwner);
+}
+
+FContentBrowserItemData UContentBrowserAliasDataSource::CreateAssetFolderItem(const FName InInternalFolderPath)
+{
+	FName VirtualizedPath;
+	TryConvertInternalPathToVirtual(InInternalFolderPath, VirtualizedPath);
+
+	const FString FolderItemName = FPackageName::GetShortName(InInternalFolderPath);
+	return FContentBrowserItemData(
+		this,
+		EContentBrowserItemFlags::Type_Folder | EContentBrowserItemFlags::Category_Asset,
+		VirtualizedPath,
+		*FolderItemName,
+		AliasFolderDisplayNames.FindRef(InInternalFolderPath),
+		MakeShared<FContentBrowserAssetFolderItemDataPayload>(InInternalFolderPath),
+		{ InInternalFolderPath });
+}
+
+FContentBrowserItemData UContentBrowserAliasDataSource::CreateAssetFileItem(const FContentBrowserUniqueAlias& Alias)
+{
+	const FAliasData* AliasData = AllAliases.Find(Alias);
+	// See UContentBrowserAliasDataSource::DoesItemPassFilter for more information on how this could fail
+	if (AliasData)
+	{
+		LLM_SCOPE(ELLMTag::UI);
+		FName VirtualizedPath;
+		PRAGMA_DISABLE_DEPRECATION_WARNINGS
+		FName InternalPath = AliasData->ObjectPath.ToFName();
+		PRAGMA_ENABLE_DEPRECATION_WARNINGS
+		TryConvertInternalPathToVirtual(InternalPath, VirtualizedPath);
+
+		// Since AliasID is PackagePath/AssetName, AssetName should also be passed as the ItemName here. This provides the functionality of
+		// being able to have multiple aliases with the same display name, while still showing their original asset name in the tooltip.
+		return FContentBrowserItemData(
+			this,
+			EContentBrowserItemFlags::Type_File | EContentBrowserItemFlags::Category_Asset,
+			VirtualizedPath,
+			AliasData->AssetData.AssetName,
+			AliasData->AliasDisplayName,
+			MakeShared<FContentBrowserAliasItemDataPayload>(AliasData->AssetData, Alias),
+			{ InternalPath });
+	}
+	return FContentBrowserItemData();
+}
diff --git a/Engine/Plugins/Editor/ContentBrowser/ContentBrowserAliasDataSource/Source/ContentBrowserAliasDataSource/Private/ContentBrowserAliasDataSourceModule.cpp b/Engine/Plugins/Editor/ContentBrowser/ContentBrowserAliasDataSource/Source/ContentBrowserAliasDataSource/Private/ContentBrowserAliasDataSourceModule.cpp
new file mode 100644
--- /dev/null
+++ b/Engine/Plugins/Editor/ContentBrowser/ContentBrowserAliasDataSource/Source/ContentBrowserAliasDataSource/Private/ContentBrowserAliasDataSourceModule.cpp
@@ -0,0 +1,35 @@
+// Copyright Epic Games, Inc. All Rights Reserved.
+
+#include "ContentBrowserAliasDataSourceModule.h"
+
+void FContentBrowserAliasDataSourceModule::StartupModule()
+{
+	UObject::GetMetaDataTagsForAssetRegistry().Add(UContentBrowserAliasDataSource::AliasTagName);
+
+	AliasDataSource.Reset(NewObject<UContentBrowserAliasDataSource>(GetTransientPackage(), "AliasData"));
+	AliasDataSource->Initialize();
+}
+
+void FContentBrowserAliasDataSourceModule::PreUnloadCallback()
+{
+	AliasDataSource.Reset();
+}
+
+UContentBrowserAliasDataSource* FContentBrowserAliasDataSourceModule::TryGetAliasDataSource()
+{
+	// GExitPurge guard required because crash inspecting object after GExitPurge
+	if (!GExitPurge)
+	{
+		if (UContentBrowserAliasDataSource* Obj = AliasDataSource.Get())
+		{
+			if (IsValid(Obj))
+			{
+				return Obj;
+			}
+		}
+	}
+
+	return nullptr;
+}
+
+IMPLEMENT_MODULE(FContentBrowserAliasDataSourceModule, ContentBrowserAliasDataSource);
diff --git a/Engine/Plugins/Editor/ContentBrowser/ContentBrowserAliasDataSource/Source/ContentBrowserAliasDataSource/Public/ContentBrowserAliasDataSource.h b/Engine/Plugins/Editor/ContentBrowser/ContentBrowserAliasDataSource/Source/ContentBrowserAliasDataSource/Public/ContentBrowserAliasDataSource.h
new file mode 100644
--- /dev/null
+++ b/Engine/Plugins/Editor/ContentBrowser/ContentBrowserAliasDataSource/Source/ContentBrowserAliasDataSource/Public/ContentBrowserAliasDataSource.h
@@ -0,0 +1,249 @@
+// Copyright Epic Games, Inc. All Rights Reserved.
+
+#pragma once
+
+
+#include "ContentBrowserAssetDataPayload.h"
+#include "ContentBrowserAssetDataSource.h"
+#include "ContentBrowserDataSource.h"
+#include "ContentBrowserLocalizedAlias.h"
+
+#include "AssetRegistry/PathTree.h"
+#include "ContentBrowserAliasDataSource.generated.h"
+
+class IAssetRegistry;
+struct FPropertyChangedEvent;
+
+DECLARE_LOG_CATEGORY_EXTERN(LogContentBrowserAliasDataSource, Log, All);
+
+class IAssetTools;
+struct FContentBrowserCompiledAssetDataFilter;
+
+/** A unique alias is a pair of SourceObjectPath:AliasPath, eg /Game/MyAsset.MyAsset:/Game/SomeFolder/MyAlias */
+typedef TPair<FSoftObjectPath, FName> FContentBrowserUniqueAlias;
+
+class CONTENTBROWSERALIASDATASOURCE_API FContentBrowserAliasItemDataPayload : public FContentBrowserAssetFileItemDataPayload
+{
+public:
+	FContentBrowserAliasItemDataPayload(const FAssetData& InAssetData, const FContentBrowserUniqueAlias& InAlias)
+		: FContentBrowserAssetFileItemDataPayload(InAssetData), Alias(InAlias)
+	{
+	}
+
+	FContentBrowserUniqueAlias Alias;
+};
+
+/**
+ * A companion to the ContentBrowserAssetDataSource which can display assets in folders other than their actual folder. Aliases mimic their source asset as closely as possible,
+ * including editing, saving, thumbnails, and more. Some behavior is restricted such as moving or deleting an alias item.
+ * 
+ * Aliases can either be created automatically by tagging the asset with the value defined by AliasTagName and giving it a comma-separated list of aliases,
+ * or manually managed by calling AddAlias and RemoveAlias. ReconcileAliasesForAsset is provided as a helper function to automatically update new/removed aliases for an existing asset.
+ *
+ */
+UCLASS()
+class CONTENTBROWSERALIASDATASOURCE_API UContentBrowserAliasDataSource : public UContentBrowserDataSource
+{
+	GENERATED_BODY()
+
+public:
+	/** The metadata tag to set for the AliasDataSource to automatically create aliases for an asset */
+	static FName AliasTagName;
+
+	// ~ Begin UContentBrowserDataSource interface
+	void Initialize(const bool InAutoRegister = true);
+	virtual void Shutdown() override;
+
+	virtual void CompileFilter(const FName InPath, const FContentBrowserDataFilter& InFilter, FContentBrowserDataCompiledFilter& OutCompiledFilter) override;
+	virtual void EnumerateItemsMatchingFilter(const FContentBrowserDataCompiledFilter& InFilter, TFunctionRef<bool(FContentBrowserItemData&&)> InCallback) override;
+	virtual void EnumerateItemsAtPath(const FName InPath, const EContentBrowserItemTypeFilter InItemTypeFilter, TFunctionRef<bool(FContentBrowserItemData&&)> InCallback) override;
+
+	virtual bool DoesItemPassFilter(const FContentBrowserItemData& InItem, const FContentBrowserDataCompiledFilter& InFilter) override;
+
+	using UContentBrowserDataSource::GetAliasesForPath;
+	virtual TArray<FContentBrowserItemPath> GetAliasesForPath(const FSoftObjectPath& InInternalPath) const override;
+
+	virtual bool GetItemAttribute(const FContentBrowserItemData& InItem, const bool InIncludeMetaData, const FName InAttributeKey, FContentBrowserItemDataAttributeValue& OutAttributeValue) override;
+	virtual bool GetItemAttributes(const FContentBrowserItemData& InItem, const bool InIncludeMetaData, FContentBrowserItemDataAttributeValues& OutAttributeValues) override;
+	virtual bool GetItemPhysicalPath(const FContentBrowserItemData& InItem, FString& OutDiskPath) override;
+	virtual bool IsItemDirty(const FContentBrowserItemData& InItem) override;
+
+	virtual bool CanEditItem(const FContentBrowserItemData& InItem, FText* OutErrorMsg) override;
+	virtual bool EditItem(const FContentBrowserItemData& InItem) override;
+	virtual bool BulkEditItems(TArrayView<const FContentBrowserItemData> InItems) override;
+	
+	virtual bool CanViewItem(const FContentBrowserItemData& InItem, FText* OutErrorMsg) override;
+	virtual bool ViewItem(const FContentBrowserItemData& InItem) override;
+	virtual bool BulkViewItems(TArrayView<const FContentBrowserItemData> InItems) override;
+
+	virtual bool CanPreviewItem(const FContentBrowserItemData& InItem, FText* OutErrorMsg) override;
+	virtual bool PreviewItem(const FContentBrowserItemData& InItem) override;
+	virtual bool BulkPreviewItems(TArrayView<const FContentBrowserItemData> InItems) override;
+
+	virtual bool CanSaveItem(const FContentBrowserItemData& InItem, const EContentBrowserItemSaveFlags InSaveFlags, FText* OutErrorMsg) override;
+	virtual bool SaveItem(const FContentBrowserItemData& InItem, const EContentBrowserItemSaveFlags InSaveFlags) override;
+	virtual bool BulkSaveItems(TArrayView<const FContentBrowserItemData> InItems, const EContentBrowserItemSaveFlags InSaveFlags) override;
+
+	virtual bool CanPrivatizeItem(const FContentBrowserItemData& InItem, FText* OutErrorMsg) override;
+	virtual bool PrivatizeItem(const FContentBrowserItemData& InItem) override;
+	virtual bool BulkPrivatizeItems(TArrayView<const FContentBrowserItemData> InItems) override;
+
+	virtual bool AppendItemReference(const FContentBrowserItemData& InItem, FString& InOutStr) override;
+	virtual bool AppendItemObjectPath(const FContentBrowserItemData& InItem, FString& InOutStr) override;
+	virtual bool AppendItemPackageName(const FContentBrowserItemData& InItem, FString& InOutStr) override;
+	virtual bool UpdateThumbnail(const FContentBrowserItemData& InItem, FAssetThumbnail& InThumbnail) override;
+
+	virtual bool TryGetCollectionId(const FContentBrowserItemData& InItem, FSoftObjectPath& OutCollectionId) override;
+
+	// Legacy functions seem necessary for FrontendFilters to work
+	virtual bool Legacy_TryGetPackagePath(const FContentBrowserItemData& InItem, FName& OutPackagePath) override;
+	virtual bool Legacy_TryGetAssetData(const FContentBrowserItemData& InItem, FAssetData& OutAssetData) override;
+	virtual bool Legacy_TryConvertPackagePathToVirtualPath(const FName InPackagePath, FName& OutPath) override;
+	virtual bool Legacy_TryConvertAssetDataToVirtualPath(const FAssetData& InAssetData, const bool InUseFolderPaths, FName& OutPath) override;
+	// ~ End UContentBrowserDataSource interface
+
+	virtual void RemoveUnusedCachedFilterData(const FContentBrowserDataFilterCacheIDOwner& IDOwner, TArrayView<const FName> InVirtualPathsInUse, const FContentBrowserDataFilter& DataFilter) override;
+
+	virtual void ClearCachedFilterData(const FContentBrowserDataFilterCacheIDOwner& IDOwner) override;
+
+	/** 
+	 * Add a list of aliases for a given asset.
+	 *
+	 * @param Asset 						The asset to add aliases for.
+	 * @param Aliases 						The aliases to add
+	 * @param bInIsFromMetaData 			Should only be `true` if the list of aliases came from the `AliasTagName` metadata.
+	 * @param bSkipPrimaryAssetValidation 	Use this if the asset being added is not a primary asset/is a re-director but should still be included in the Content Browser. (e.g. Verse classes)
+	 */
+	void AddAliases(const FAssetData& Asset, const TArray<FName>& Aliases, const bool bInIsFromMetaData = false, const bool bSkipPrimaryAssetValidation = false);
+	void AddAliases(const FAssetData& Asset, const TArray<FContentBrowserLocalizedAlias>& Aliases, const bool bInIsFromMetaData = false, const bool bSkipPrimaryAssetValidation = false);
+	/** Add an alias for a given asset. bInIsFromMetaData should only be true if the alias came from the AliasTagName metadata. */
+	void AddAlias(const FAssetData& Asset, const FName Alias, const bool bInIsFromMetaData = false, const bool bSkipPrimaryAssetValidation = false);
+	void AddAlias(const FAssetData& Asset, const FContentBrowserLocalizedAlias& Alias, const bool bInIsFromMetaData = false, const bool bSkipPrimaryAssetValidation = false);
+	/** Remove the given alias from the data source */
+	void RemoveAlias(const FSoftObjectPath& ObjectPath, const FName Alias);
+	/** Remove all aliases for the given object */
+	void RemoveAliases(const FSoftObjectPath& ObjectPath);
+	/** Remove all aliases for the given asset */
+	void RemoveAliases(const FAssetData& Asset) { RemoveAliases(Asset.GetSoftObjectPath()); }
+
+	UE_DEPRECATED(5.1, "FNames containing full asset paths are deprecated, use FSoftObjectPath instead")
+	void RemoveAlias(const FName ObjectPath, const FName Alias);
+	UE_DEPRECATED(5.1, "FNames containing full asset paths are deprecated, use FSoftObjectPath instead")
+	void RemoveAliases(const FName ObjectPath);
+
+	/**
+	 * Add a display name override for the given alias folder, eg /MyAliases.
+	 * @note To provide a display name override for an alias itself, use the overloads that take a FContentBrowserLocalizedAlias.
+	 */
+	void AddAliasFolderDisplayName(const FName AliasFolder, const FText& DisplayName);
+	/**
+	 * Remove a display name override for the given alias folder, eg /MyAliases.
+	 */
+	void RemoveAliasFolderDisplayName(const FName AliasFolder);
+
+	/** When called, removes all aliases and triggers delegate for various systems to re-add aliases */
+	void RebuildAliases();
+
+	/** Broadcast after RebuildAliases() called to allow systems to re-add aliases */
+	FSimpleMulticastDelegate& OnRebuildAliases() { return RebuildAliasesDelegate; }
+
+	/** Get all aliases from metadata for the given asset, then calls AddAlias or RemoveAlias for every alias that doesn't match the stored data. */
+	void ReconcileAliasesFromMetaData(const FAssetData& Asset);
+	/** Calls AddAlias or RemoveAlias for every alias that doesn't match the stored data for the given asset. */
+	void ReconcileAliasesForAsset(const FAssetData& Asset, const TArray<FName>& NewAliases);
+	void ReconcileAliasesForAsset(const FAssetData& Asset, const TArray<FContentBrowserLocalizedAlias>& NewAliases);
+
+	/** Logs all the content browser aliases */
+	void LogAliases() const;
+
+protected:
+	virtual void BuildRootPathVirtualTree() override;
+
+private:
+	template <typename AliasType>
+	void AddAliasImpl(const FAssetData& Asset, const AliasType& Alias, const bool bInIsFromMetaData, const bool bSkipPrimaryAssetValidation);
+	template <typename AliasType>
+	void AddAliasesImpl(const FAssetData& Asset, const TArray<AliasType>& Aliases, const bool bInIsFromMetaData, const bool bSkipPrimaryAssetValidation);
+	template <typename AliasType>
+	void ReconcileAliasesForAssetImpl(const FAssetData& Asset, const TArray<AliasType>& NewAliases);
+
+	void OnAssetAdded(const FAssetData& InAssetData);
+	void OnAssetRemoved(const FAssetData& InAssetData);
+	void OnAssetUpdated(const FAssetData& InAssetData);
+	void OnAssetLoaded(UObject* InAsset);
+	void OnObjectPropertyChanged(UObject* InObject, FPropertyChangedEvent& InPropertyChangedEvent);
+
+	/** Helper function to remove a folder from the PathTree including all parent folders that are now empty as a result of the removal */
+	void RemoveFoldersRecursively(FStringView LeafFolder);
+
+	void UpdateAliasesCachedAssetData(const FAssetData& InAssetData);
+	void MakeItemModifiedUpdate(const FSoftObjectPath& ObjectPath);
+
+	FContentBrowserItemData CreateAssetFolderItem(const FName InFolderPath);
+	FContentBrowserItemData CreateAssetFileItem(const FContentBrowserUniqueAlias& Alias);
+
+	IAssetRegistry* AssetRegistry = nullptr;
+	IAssetTools* AssetTools = nullptr;
+
+	struct FAliasData
+	{
+		FAliasData() {}
+		FAliasData(const FAssetData& InAssetData, const FName InPackagePath, const FText& InDisplayName, const bool bInIsFromMetaData = false)
+			: AssetData(InAssetData), PackagePath(InPackagePath), AliasDisplayName(InDisplayName), bIsFromMetaData(bInIsFromMetaData)
+		{
+			FNameBuilder AssetNameBuilder(InAssetData.AssetName);
+			{
+				// Add a hash of the real asset path to the asset name to ensure uniqueness of the FContentBrowserItemKey
+				// so that two different assets with the same name and the same alias both show up in the content browser
+				FNameBuilder AssetPathBuilder;
+				InAssetData.AppendObjectPath(AssetPathBuilder);
+				AssetNameBuilder.Appendf(TEXT("_%08X"), GetTypeHash(AssetPathBuilder.ToView()));
+			}
+			{
+				FNameBuilder PathBuilder(PackagePath);
+				PathBuilder << TEXT('/');
+				PathBuilder << AssetNameBuilder.ToView();
+				PackageName = FName(PathBuilder.ToView());
+			}
+			ObjectPath = FSoftObjectPath(PackageName, FName(AssetNameBuilder.ToView()), {});
+		}
+
+		/** The source asset for this alias */
+		FAssetData AssetData;
+		/** The folder path that contains the alias, /MyAliases */
+		FName PackagePath;
+		/** PackagePath/SourceAssetName, /MyAliases/SomeAsset  */
+		FName PackageName;
+		/** PackageName.SourceAssetName, /MyAliases/SomeAsset.SomeAsset */
+		FSoftObjectPath ObjectPath;
+		/** A non-unique display name for this alias */
+		FText AliasDisplayName;
+		/** Whether this alias was generated from package metadata or manually through the C++ interface */
+		bool bIsFromMetaData = false;
+	};
+	bool DoesAliasPassFilter(const FAliasData& AliasData, const FContentBrowserCompiledAssetDataFilter& Filter) const;
+
+	/** The full folder hierarchy for all alias paths */
+	FPathTree PathTree;
+	/** Alias data keyed by their full alias path, eg /Game/MyData/Aliases/SourceMesh */
+	TMap<FContentBrowserUniqueAlias, FAliasData> AllAliases;
+	/** A list of alias paths to display for each asset, eg /Game/Meshes/SourceMesh.SourceMesh */
+	TMap<FSoftObjectPath, TArray<FName>> AliasesForObjectPath;
+	/** A list of alias paths to display for each folder, eg /Game/MyData/Aliases */
+	TMap<FName, TArray<FContentBrowserUniqueAlias>> AliasesInPackagePath;
+	/** Alias folder display names keyed against their alias path, eg /Game/MyData/Aliases/ */
+	TMap<FName, FText> AliasFolderDisplayNames;
+	/** A set used for removing duplicate aliases in the same query, stored here to avoid constant reallocation */
+	TSet<FSoftObjectPath> AlreadyAddedOriginalAssets;
+
+	UContentBrowserAssetDataSource::FAssetDataSourceFilterCache FilterCache;
+
+	/** Delegate broadcast after all aliases removed to give chance for systems to re-add aliases */
+	FSimpleMulticastDelegate RebuildAliasesDelegate;
+};
+
+#if UE_ENABLE_INCLUDE_ORDER_DEPRECATED_IN_5_2
+#include "AssetRegistry/AssetRegistryModule.h"
+#include "ContentBrowserItemPath.h"
+#include "CoreMinimal.h"
+#endif
diff --git a/Engine/Plugins/Editor/ContentBrowser/ContentBrowserAliasDataSource/Source/ContentBrowserAliasDataSource/Public/ContentBrowserAliasDataSourceModule.h b/Engine/Plugins/Editor/ContentBrowser/ContentBrowserAliasDataSource/Source/ContentBrowserAliasDataSource/Public/ContentBrowserAliasDataSourceModule.h
new file mode 100644
--- /dev/null
+++ b/Engine/Plugins/Editor/ContentBrowser/ContentBrowserAliasDataSource/Source/ContentBrowserAliasDataSource/Public/ContentBrowserAliasDataSourceModule.h
@@ -0,0 +1,22 @@
+// Copyright Epic Games, Inc. All Rights Reserved.
+
+#pragma once
+
+#include "ContentBrowserAliasDataSource.h"
+#include "Modules/ModuleManager.h"
+#include "UObject/StrongObjectPtr.h"
+
+class FContentBrowserAliasDataSourceModule : public FDefaultModuleImpl
+{
+public:
+	virtual void StartupModule() override;
+	virtual void PreUnloadCallback() override;
+
+	UE_DEPRECATED(5.4, "Call TryGetAliasDataSource instead.")
+	TWeakObjectPtr<UContentBrowserAliasDataSource> GetAliasDataSource() { return AliasDataSource.Get(); }
+
+	CONTENTBROWSERALIASDATASOURCE_API UContentBrowserAliasDataSource* TryGetAliasDataSource();
+
+private:
+	TStrongObjectPtr<UContentBrowserAliasDataSource> AliasDataSource;
+};
diff --git a/Engine/Plugins/Editor/ContentBrowser/ContentBrowserAliasDataSource/Source/ContentBrowserAliasDataSource/Public/ContentBrowserLocalizedAlias.h b/Engine/Plugins/Editor/ContentBrowser/ContentBrowserAliasDataSource/Source/ContentBrowserAliasDataSource/Public/ContentBrowserLocalizedAlias.h
new file mode 100644
--- /dev/null
+++ b/Engine/Plugins/Editor/ContentBrowser/ContentBrowserAliasDataSource/Source/ContentBrowserAliasDataSource/Public/ContentBrowserLocalizedAlias.h
@@ -0,0 +1,28 @@
+// Copyright Epic Games, Inc. All Rights Reserved.
+
+#pragma once
+
+#include "Internationalization/Text.h"
+#include "UObject/NameTypes.h"
+
+/** Container for an alias and its optional display name override */
+struct FContentBrowserLocalizedAlias
+{
+	FName Alias;
+	FText DisplayName;
+
+	bool operator==(const FContentBrowserLocalizedAlias& Other) const
+	{
+		return Alias == Other.Alias;
+	}
+
+	bool operator!=(const FContentBrowserLocalizedAlias& Other) const
+	{
+		return Alias != Other.Alias;
+	}
+
+	friend uint32 GetTypeHash(const FContentBrowserLocalizedAlias& This)
+	{
+		return GetTypeHash(This.Alias);
+	}
+};
diff --git a/Engine/Plugins/Editor/ContentBrowser/ContentBrowserAliasDataSoure/ContentBrowserAliasDataSource.uplugin b/Engine/Plugins/Editor/ContentBrowser/ContentBrowserAliasDataSoure/ContentBrowserAliasDataSource.uplugin
deleted file mode 100644
diff --git a/Engine/Plugins/Editor/ContentBrowser/ContentBrowserAliasDataSoure/Source/ContentBrowserAliasDataSource/ContentBrowserAliasDataSource.Build.cs b/Engine/Plugins/Editor/ContentBrowser/ContentBrowserAliasDataSoure/Source/ContentBrowserAliasDataSource/ContentBrowserAliasDataSource.Build.cs
deleted file mode 100644
diff --git a/Engine/Plugins/Editor/ContentBrowser/ContentBrowserAliasDataSoure/Source/ContentBrowserAliasDataSource/Private/ContentBrowserAliasDataSource.cpp b/Engine/Plugins/Editor/ContentBrowser/ContentBrowserAliasDataSoure/Source/ContentBrowserAliasDataSource/Private/ContentBrowserAliasDataSource.cpp
deleted file mode 100644
diff --git a/Engine/Plugins/Editor/ContentBrowser/ContentBrowserAliasDataSoure/Source/ContentBrowserAliasDataSource/Private/ContentBrowserAliasDataSourceModule.cpp b/Engine/Plugins/Editor/ContentBrowser/ContentBrowserAliasDataSoure/Source/ContentBrowserAliasDataSource/Private/ContentBrowserAliasDataSourceModule.cpp
deleted file mode 100644
diff --git a/Engine/Plugins/Editor/ContentBrowser/ContentBrowserAliasDataSoure/Source/ContentBrowserAliasDataSource/Public/ContentBrowserAliasDataSource.h b/Engine/Plugins/Editor/ContentBrowser/ContentBrowserAliasDataSoure/Source/ContentBrowserAliasDataSource/Public/ContentBrowserAliasDataSource.h
deleted file mode 100644
diff --git a/Engine/Plugins/Editor/ContentBrowser/ContentBrowserAliasDataSoure/Source/ContentBrowserAliasDataSource/Public/ContentBrowserAliasDataSourceModule.h b/Engine/Plugins/Editor/ContentBrowser/ContentBrowserAliasDataSoure/Source/ContentBrowserAliasDataSource/Public/ContentBrowserAliasDataSourceModule.h
deleted file mode 100644
diff --git a/Engine/Plugins/Editor/ContentBrowser/ContentBrowserAliasDataSoure/Source/ContentBrowserAliasDataSource/Public/ContentBrowserLocalizedAlias.h b/Engine/Plugins/Editor/ContentBrowser/ContentBrowserAliasDataSoure/Source/ContentBrowserAliasDataSource/Public/ContentBrowserLocalizedAlias.h
deleted file mode 100644
