Allow querying whether a given source asset has any aliases

#rnx

diff --git a/Engine/Plugins/Editor/ContentBrowser/ContentBrowserAliasDataSource/Source/ContentBrowserAliasDataSource/Private/ContentBrowserAliasDataSource.cpp b/Engine/Plugins/Editor/ContentBrowser/ContentBrowserAliasDataSource/Source/ContentBrowserAliasDataSource/Private/ContentBrowserAliasDataSource.cpp
--- a/Engine/Plugins/Editor/ContentBrowser/ContentBrowserAliasDataSource/Source/ContentBrowserAliasDataSource/Private/ContentBrowserAliasDataSource.cpp
+++ b/Engine/Plugins/Editor/ContentBrowser/ContentBrowserAliasDataSource/Source/ContentBrowserAliasDataSource/Private/ContentBrowserAliasDataSource.cpp
@@ -338,6 +338,11 @@ TArray<FContentBrowserItemPath> UContentBrowserAliasDataSource::GetAliasesForPat
 	return OutAliases;
 }
 
+bool UContentBrowserAliasDataSource::HasAliasesForPath(const FSoftObjectPath& ObjectPath) const
+{
+	return AliasesForObjectPath.Contains(ObjectPath);
+}
+
 void UContentBrowserAliasDataSource::AddAliases(const FAssetData& Asset, const TArray<FName>& Aliases, const bool bInIsFromMetaData, const bool bSkipPrimaryAssetValidation)
 {
 	AddAliasesImpl(Asset, Aliases, bInIsFromMetaData, bSkipPrimaryAssetValidation);
diff --git a/Engine/Plugins/Editor/ContentBrowser/ContentBrowserAliasDataSource/Source/ContentBrowserAliasDataSource/Public/ContentBrowserAliasDataSource.h b/Engine/Plugins/Editor/ContentBrowser/ContentBrowserAliasDataSource/Source/ContentBrowserAliasDataSource/Public/ContentBrowserAliasDataSource.h
--- a/Engine/Plugins/Editor/ContentBrowser/ContentBrowserAliasDataSource/Source/ContentBrowserAliasDataSource/Public/ContentBrowserAliasDataSource.h
+++ b/Engine/Plugins/Editor/ContentBrowser/ContentBrowserAliasDataSource/Source/ContentBrowserAliasDataSource/Public/ContentBrowserAliasDataSource.h
@@ -63,6 +63,8 @@ public:
 	using UContentBrowserDataSource::GetAliasesForPath;
 	virtual TArray<FContentBrowserItemPath> GetAliasesForPath(const FSoftObjectPath& InInternalPath) const override;
 
+	bool HasAliasesForPath(const FSoftObjectPath& InInternalPath) const;
+
 	virtual bool GetItemAttribute(const FContentBrowserItemData& InItem, const bool InIncludeMetaData, const FName InAttributeKey, FContentBrowserItemDataAttributeValue& OutAttributeValue) override;
 	virtual bool GetItemAttributes(const FContentBrowserItemData& InItem, const bool InIncludeMetaData, FContentBrowserItemDataAttributeValues& OutAttributeValues) override;
 	virtual bool GetItemPhysicalPath(const FContentBrowserItemData& InItem, FString& OutDiskPath) override;
