Fixed empty folder filters not working under mount points with aliased data sources.

UContentBrowserAliasDataSource was using the default IsFolderVisible implementation, which always returns true.  As a result, UContentBrowserDataSubsystem::IsFolderVisible would return true for content under mount points that contain aliased content, which was an issue when different content sources put content under the same mount point.

diff --git a/Engine/Plugins/Editor/ContentBrowser/ContentBrowserAliasDataSource/Source/ContentBrowserAliasDataSource/Private/ContentBrowserAliasDataSource.cpp b/Engine/Plugins/Editor/ContentBrowser/ContentBrowserAliasDataSource/Source/ContentBrowserAliasDataSource/Private/ContentBrowserAliasDataSource.cpp
--- a/Engine/Plugins/Editor/ContentBrowser/ContentBrowserAliasDataSource/Source/ContentBrowserAliasDataSource/Private/ContentBrowserAliasDataSource.cpp
+++ b/Engine/Plugins/Editor/ContentBrowser/ContentBrowserAliasDataSource/Source/ContentBrowserAliasDataSource/Private/ContentBrowserAliasDataSource.cpp
@@ -292,6 +292,12 @@ void UContentBrowserAliasDataSource::EnumerateItemsAtPath(const FName InPath, co
 	}
 }
 
+bool UContentBrowserAliasDataSource::IsFolderVisible(const FName Path, const EContentBrowserIsFolderVisibleFlags Flags, const FContentBrowserFolderContentsFilter& ContentsFilter)
+{
+	FName InternalPath;
+	return TryConvertVirtualPathToInternal(Path, InternalPath) && PathTree.PathExists(InternalPath);
+}
+
 bool UContentBrowserAliasDataSource::DoesItemPassFilter(const FContentBrowserItemData& InItem, const FContentBrowserDataCompiledFilter& InFilter)
 {
 	const FContentBrowserDataFilterList* FilterList = InFilter.CompiledFilters.Find(this);
diff --git a/Engine/Plugins/Editor/ContentBrowser/ContentBrowserAliasDataSource/Source/ContentBrowserAliasDataSource/Public/ContentBrowserAliasDataSource.h b/Engine/Plugins/Editor/ContentBrowser/ContentBrowserAliasDataSource/Source/ContentBrowserAliasDataSource/Public/ContentBrowserAliasDataSource.h
--- a/Engine/Plugins/Editor/ContentBrowser/ContentBrowserAliasDataSource/Source/ContentBrowserAliasDataSource/Public/ContentBrowserAliasDataSource.h
+++ b/Engine/Plugins/Editor/ContentBrowser/ContentBrowserAliasDataSource/Source/ContentBrowserAliasDataSource/Public/ContentBrowserAliasDataSource.h
@@ -58,6 +58,8 @@ public:
 	virtual void EnumerateItemsMatchingFilter(const FContentBrowserDataCompiledFilter& InFilter, TFunctionRef<bool(FContentBrowserItemData&&)> InCallback) override;
 	virtual void EnumerateItemsAtPath(const FName InPath, const EContentBrowserItemTypeFilter InItemTypeFilter, TFunctionRef<bool(FContentBrowserItemData&&)> InCallback) override;
 
+	virtual bool IsFolderVisible(const FName Path, const EContentBrowserIsFolderVisibleFlags Flags, const FContentBrowserFolderContentsFilter& ContentsFilter) override;
+
 	virtual bool DoesItemPassFilter(const FContentBrowserItemData& InItem, const FContentBrowserDataCompiledFilter& InFilter) override;
 
 	using UContentBrowserDataSource::GetAliasesForPath;
