Updated headers using UnrealCodeFixup to make sure dllstorage is on methods/staticvars instead of types. Part 2/n

diff --git a/Engine/Plugins/Editor/BlueprintHeaderView/Source/BlueprintHeaderView/Public/BlueprintHeaderViewSettings.h b/Engine/Plugins/Editor/BlueprintHeaderView/Source/BlueprintHeaderView/Public/BlueprintHeaderViewSettings.h
--- a/Engine/Plugins/Editor/BlueprintHeaderView/Source/BlueprintHeaderView/Public/BlueprintHeaderViewSettings.h
+++ b/Engine/Plugins/Editor/BlueprintHeaderView/Source/BlueprintHeaderView/Public/BlueprintHeaderViewSettings.h
@@ -5,6 +5,8 @@
 #include "Engine/DeveloperSettings.h"
 #include "BlueprintHeaderViewSettings.generated.h"
 
+#define UE_API BLUEPRINTHEADERVIEW_API
+
 UENUM()
 enum class EHeaderViewSortMethod : uint8
 {
@@ -47,23 +49,23 @@ public:
 };
 
 /** Settings for the Blueprint Header View Plugin */
-UCLASS(config = EditorPerProjectUserSettings)
-class BLUEPRINTHEADERVIEW_API UBlueprintHeaderViewSettings : public UDeveloperSettings
+UCLASS(MinimalAPI, config = EditorPerProjectUserSettings)
+class UBlueprintHeaderViewSettings : public UDeveloperSettings
 {
 	GENERATED_BODY()
 
 public:
-	UBlueprintHeaderViewSettings();
+	UE_API UBlueprintHeaderViewSettings();
 
 	//~ Begin UDeveloperSettings interface
-	virtual FName GetCategoryName() const override;
-	virtual FText GetSectionText() const override;
-	virtual FName GetSectionName() const override;
+	UE_API virtual FName GetCategoryName() const override;
+	UE_API virtual FText GetSectionText() const override;
+	UE_API virtual FName GetSectionName() const override;
 	//~ End UDeveloperSettings interface
 
 protected:
 	//~ Begin UObject interface
-	virtual void PostEditChangeProperty(struct FPropertyChangedEvent& PropertyChangedEvent) override;
+	UE_API virtual void PostEditChangeProperty(struct FPropertyChangedEvent& PropertyChangedEvent) override;
 	//~ End UObject interface
 
 public:
@@ -84,3 +86,5 @@ public:
 	UPROPERTY(config, EditAnywhere, Category="Settings")
 	EHeaderViewSortMethod SortMethod = EHeaderViewSortMethod::None;
 };
+
+#undef UE_API
diff --git a/Engine/Plugins/Editor/BlueprintMaterialTextureNodes/Source/BlueprintMaterialTextureNodes/Public/BlueprintMaterialTextureNodesBPLibrary.h b/Engine/Plugins/Editor/BlueprintMaterialTextureNodes/Source/BlueprintMaterialTextureNodes/Public/BlueprintMaterialTextureNodesBPLibrary.h
--- a/Engine/Plugins/Editor/BlueprintMaterialTextureNodes/Source/BlueprintMaterialTextureNodes/Public/BlueprintMaterialTextureNodesBPLibrary.h
+++ b/Engine/Plugins/Editor/BlueprintMaterialTextureNodes/Source/BlueprintMaterialTextureNodes/Public/BlueprintMaterialTextureNodesBPLibrary.h
@@ -6,6 +6,8 @@
 #include "Engine/EngineTypes.h"
 #include "BlueprintMaterialTextureNodesBPLibrary.generated.h"
 
+#define UE_API BLUEPRINTMATERIALTEXTURENODES_API
+
 class UMaterialInstanceConstant;
 class UMaterialInterface;
 class UTexture2D;
@@ -28,8 +30,8 @@ class UTextureRenderTarget2D;
 *	For more info on custom blueprint nodes visit documentation:
 *	https://docs.unrealengine.com/ProgrammingAndScripting/Blueprints/TechnicalGuide/ExtendingBlueprints
 */
-UCLASS()
-class BLUEPRINTMATERIALTEXTURENODES_API UBlueprintMaterialTextureNodesBPLibrary : public UBlueprintFunctionLibrary
+UCLASS(MinimalAPI)
+class UBlueprintMaterialTextureNodesBPLibrary : public UBlueprintFunctionLibrary
 {
 	GENERATED_UCLASS_BODY()
 
@@ -38,21 +40,21 @@ class BLUEPRINTMATERIALTEXTURENODES_API UBlueprintMaterialTextureNodesBPLibrary
 	* Samples a texel from a Texture 2D with VectorDisplacement Compression
 	*/
 	UFUNCTION(BlueprintPure, meta = (DisplayName = "Texture2D Sample UV Editor Only", Keywords = "Read Texture UV"), Category = Rendering)
-	static FLinearColor Texture2D_SampleUV_EditorOnly(UTexture2D* Texture, FVector2D UV);
+	static UE_API FLinearColor Texture2D_SampleUV_EditorOnly(UTexture2D* Texture, FVector2D UV);
 
 	/**
 	* Samples an array of values from a Texture Render Target 2D. Currently only 4 channel formats are supported.
 	* Only works in the editor
 	*/
 	UFUNCTION(BlueprintPure, meta = (DisplayName = "Render Target Sample Rectangle Editor Only", Keywords = "Sample Render Target Rectangle"), Category = Rendering)
-	static TArray<FLinearColor> RenderTarget_SampleRectangle_EditorOnly(UTextureRenderTarget2D* InRenderTarget, FLinearColor InRect);
+	static UE_API TArray<FLinearColor> RenderTarget_SampleRectangle_EditorOnly(UTextureRenderTarget2D* InRenderTarget, FLinearColor InRect);
 
 	/**
 	* Samples a value from a Texture Render Target 2D. Currently only 4 channel formats are supported.
 	* Only works in the editor
 	*/
 	UFUNCTION(BlueprintPure, meta = (DisplayName = "Render Target Sample UV Editor Only", Keywords = "Sample Render Target UV"), Category = Rendering)
-	static FLinearColor RenderTarget_SampleUV_EditorOnly(UTextureRenderTarget2D* InRenderTarget, FVector2D UV);
+	static UE_API FLinearColor RenderTarget_SampleUV_EditorOnly(UTextureRenderTarget2D* InRenderTarget, FVector2D UV);
 
 
 	/**
@@ -60,64 +62,66 @@ class BLUEPRINTMATERIALTEXTURENODES_API UBlueprintMaterialTextureNodesBPLibrary
 	* Only works in the editor
 	*/
 	UFUNCTION(BlueprintCallable, meta = (DisplayName = "Create MIC Editor Only", Keywords = "Create MIC", UnsafeDuringActorConstruction = "true"), Category = Rendering)
-	static UMaterialInstanceConstant* CreateMIC_EditorOnly(UMaterialInterface* Material, FString Name = "MIC_");
+	static UE_API UMaterialInstanceConstant* CreateMIC_EditorOnly(UMaterialInterface* Material, FString Name = "MIC_");
 
 	UFUNCTION()
-	static void UpdateMIC(UMaterialInstanceConstant* MIC);
+	static UE_API void UpdateMIC(UMaterialInstanceConstant* MIC);
 
 	/**
 	* Sets a Scalar Parameter value in a Material Instance Constant
 	* Only works in the editor
 	*/
 	UFUNCTION(BlueprintCallable, meta = (DisplayName = "Set MIC Scalar Parameter Editor Only", Keywords = "Set MIC Scalar Parameter", UnsafeDuringActorConstruction = "true"), Category = Rendering)
-	static bool SetMICScalarParam_EditorOnly(UMaterialInstanceConstant* Material, FString ParamName = "test", float Value = 0.0f);
+	static UE_API bool SetMICScalarParam_EditorOnly(UMaterialInstanceConstant* Material, FString ParamName = "test", float Value = 0.0f);
 
 	/**
 	* Sets a Vector Parameter value in a Material Instance Constant
 	* Only works in the editor
 	*/
 	UFUNCTION(BlueprintCallable, meta = (DisplayName = "Set MIC Vector Parameter Editor Only", Keywords = "Set MIC Vector Parameter", UnsafeDuringActorConstruction = "true"), Category = Rendering)
-	static bool SetMICVectorParam_EditorOnly(UMaterialInstanceConstant* Material, FString ParamName, FLinearColor Value = FLinearColor(0, 0, 0, 0));
+	static UE_API bool SetMICVectorParam_EditorOnly(UMaterialInstanceConstant* Material, FString ParamName, FLinearColor Value = FLinearColor(0, 0, 0, 0));
 
 	/**
 	* Sets a Texture Parameter value in a Material Instance Constant
 	* Only works in the editor
 	*/
 	UFUNCTION(BlueprintCallable, meta = (DisplayName = "Set MIC Texture Parameter Editor Only", Keywords = "Set MIC Texture Parameter", UnsafeDuringActorConstruction = "true"), Category = Rendering)
-	static bool SetMICTextureParam_EditorOnly(UMaterialInstanceConstant* Material, FString ParamName, UTexture2D* Texture = nullptr);
+	static UE_API bool SetMICTextureParam_EditorOnly(UMaterialInstanceConstant* Material, FString ParamName, UTexture2D* Texture = nullptr);
 
 	/**
 	* Overrides the Shading Model of a Material Instance Constant
 	* Only works in the editor
 	*/
 	UFUNCTION(BlueprintCallable, meta = (DisplayName = "Set MIC Shading Model Editor Only", Keywords = "Set MIC Shading Model", UnsafeDuringActorConstruction = "true"), Category = Rendering)
-	static bool SetMICShadingModel_EditorOnly(UMaterialInstanceConstant* Material, TEnumAsByte<EMaterialShadingModel> ShadingModel = MSM_DefaultLit);
+	static UE_API bool SetMICShadingModel_EditorOnly(UMaterialInstanceConstant* Material, TEnumAsByte<EMaterialShadingModel> ShadingModel = MSM_DefaultLit);
 
 	/**
 	* Overrides the Blend Mode of a Material Instance Constant
 	* Only works in the editor
 	*/
 	UFUNCTION(BlueprintCallable, meta = (DisplayName = "Set MIC Blend Mode Editor Only", Keywords = "Set MIC Blend Mode", UnsafeDuringActorConstruction = "true"), Category = Rendering)
-	static bool SetMICBlendMode_EditorOnly(UMaterialInstanceConstant* Material, TEnumAsByte<EBlendMode> BlendMode = BLEND_Opaque);
+	static UE_API bool SetMICBlendMode_EditorOnly(UMaterialInstanceConstant* Material, TEnumAsByte<EBlendMode> BlendMode = BLEND_Opaque);
 
 	/**
 	* Overrides the Two Sided setting of a Material Instance Constant
 	* Only works in the editor
 	*/
 	UFUNCTION(BlueprintCallable, meta = (DisplayName = "Set MIC Two Sided Editor Only", Keywords = "Set MIC Two Sided", UnsafeDuringActorConstruction = "true"), Category = Rendering)
-	static bool SetMICTwoSided_EditorOnly(UMaterialInstanceConstant* Material, bool TwoSided = false);
+	static UE_API bool SetMICTwoSided_EditorOnly(UMaterialInstanceConstant* Material, bool TwoSided = false);
 
 	/**
 	* Overrides the IsThinSurface setting of a Material Instance Constant
 	* Only works in the editor
 	*/
 	UFUNCTION(BlueprintCallable, meta = (DisplayName = "Set MIC IsThinSurface Editor Only", Keywords = "Set MIC IsThinSurface", UnsafeDuringActorConstruction = "true"), Category = Rendering)
-	static bool SetMICIsThinSurface_EditorOnly(UMaterialInstanceConstant* Material, bool bIsThinSurface = false);	
+	static UE_API bool SetMICIsThinSurface_EditorOnly(UMaterialInstanceConstant* Material, bool bIsThinSurface = false);	
 
 	/**
 	* Overrides the Blend Mode of a Material Instance Constant
 	* Only works in the editor
 	*/
 	UFUNCTION(BlueprintCallable, meta = (DisplayName = "Set MIC Dithered LOD Editor Only", Keywords = "Set MIC Dithered LOD Transition", UnsafeDuringActorConstruction = "true"), Category = Rendering)
-	static bool SetMICDitheredLODTransition_EditorOnly(UMaterialInstanceConstant* Material, bool DitheredLODTransition = false);
+	static UE_API bool SetMICDitheredLODTransition_EditorOnly(UMaterialInstanceConstant* Material, bool DitheredLODTransition = false);
 };
+
+#undef UE_API
diff --git a/Engine/Plugins/Editor/ChangelistReview/Source/Public/ChangelistReviewModule.h b/Engine/Plugins/Editor/ChangelistReview/Source/Public/ChangelistReviewModule.h
--- a/Engine/Plugins/Editor/ChangelistReview/Source/Public/ChangelistReviewModule.h
+++ b/Engine/Plugins/Editor/ChangelistReview/Source/Public/ChangelistReviewModule.h
@@ -5,6 +5,8 @@
 //#include "SSourceControlReview.h"
 #include "Modules/ModuleManager.h"
 
+#define UE_API CHANGELISTREVIEW_API
+
 class SSourceControlReview;
 class FSpawnTabArgs;
 class SDockTab;
@@ -12,14 +14,14 @@ class SWidget;
 
 class UContentBrowserAliasDataSource;
 
-class CHANGELISTREVIEW_API FChangelistReviewModule : public FDefaultModuleImpl
+class FChangelistReviewModule : public FDefaultModuleImpl
 {
 public:
-	virtual void StartupModule() override;
-	virtual void ShutdownModule() override;
-	void ShowReviewTab();
-	bool CanShowReviewTab() const;
-	TWeakPtr<SSourceControlReview> GetActiveReview();
+	UE_API virtual void StartupModule() override;
+	UE_API virtual void ShutdownModule() override;
+	UE_API void ShowReviewTab();
+	UE_API bool CanShowReviewTab() const;
+	UE_API TWeakPtr<SSourceControlReview> GetActiveReview();
 	
 	/**
 	 * Opens review tool and loads the provided change list.
@@ -27,16 +29,18 @@ public:
 	 * @param[in] Changelist The change list to load.
 	 * @return @c true if changelist review tool was opened and change list is put to be loaded, @c false otherwise.
 	 */
-	bool OpenChangelistReview(const FString& Changelist);
+	UE_API bool OpenChangelistReview(const FString& Changelist);
 
 	static FChangelistReviewModule& Get()
 	{
 		return FModuleManager::LoadModuleChecked<FChangelistReviewModule>("ChangelistReview");
 	}
 private:
-	TSharedRef<SDockTab> CreateReviewTab(const FSpawnTabArgs& Args);
-	TSharedPtr<SWidget> CreateReviewUI();
+	UE_API TSharedRef<SDockTab> CreateReviewTab(const FSpawnTabArgs& Args);
+	UE_API TSharedPtr<SWidget> CreateReviewUI();
 	
 	TWeakPtr<SDockTab> ReviewTab;
 	TWeakPtr<SSourceControlReview> ReviewWidget;
 };
+
+#undef UE_API
diff --git a/Engine/Plugins/Editor/ColorGrading/Source/ColorGradingEditor/Private/ColorGradingMixerObjectFilter.h b/Engine/Plugins/Editor/ColorGrading/Source/ColorGradingEditor/Private/ColorGradingMixerObjectFilter.h
--- a/Engine/Plugins/Editor/ColorGrading/Source/ColorGradingEditor/Private/ColorGradingMixerObjectFilter.h
+++ b/Engine/Plugins/Editor/ColorGrading/Source/ColorGradingEditor/Private/ColorGradingMixerObjectFilter.h
@@ -6,21 +6,23 @@
 
 #include "ColorGradingMixerObjectFilter.generated.h"
 
-UCLASS(BlueprintType, EditInlineNew)
-class COLORGRADINGEDITOR_API UColorGradingMixerObjectFilter : public UObjectMixerObjectFilter
+#define UE_API COLORGRADINGEDITOR_API
+
+UCLASS(MinimalAPI, BlueprintType, EditInlineNew)
+class UColorGradingMixerObjectFilter : public UObjectMixerObjectFilter
 {
 	GENERATED_BODY()
 public:
 
-	virtual TSet<UClass*> GetObjectClassesToFilter() const override;
-	virtual TSet<TSubclassOf<AActor>> GetObjectClassesToPlace() const override;
-	virtual TArray<AActor*> FindAssociatedActors(AActor* InActor) const override;
-	virtual bool IsActorAssociated(AActor* Actor, AActor* AssociatedActor) const override;
-	virtual bool HasCustomDropHandling(const ISceneOutlinerTreeItem& DropTarget) const override;
-	virtual FSceneOutlinerDragValidationInfo ValidateDrop(const ISceneOutlinerTreeItem& DropTarget, const FSceneOutlinerDragDropPayload& Payload) const override;
-	virtual void OnDrop(ISceneOutlinerTreeItem& DropTarget, const FSceneOutlinerDragDropPayload& Payload, const FSceneOutlinerDragValidationInfo& ValidationInfo) const override;
-	virtual TSet<FName> GetPropertiesThatRequireListRefresh() const override;
-	virtual void OnContextMenuContextCreated(FToolMenuContext& Context) const override;
+	UE_API virtual TSet<UClass*> GetObjectClassesToFilter() const override;
+	UE_API virtual TSet<TSubclassOf<AActor>> GetObjectClassesToPlace() const override;
+	UE_API virtual TArray<AActor*> FindAssociatedActors(AActor* InActor) const override;
+	UE_API virtual bool IsActorAssociated(AActor* Actor, AActor* AssociatedActor) const override;
+	UE_API virtual bool HasCustomDropHandling(const ISceneOutlinerTreeItem& DropTarget) const override;
+	UE_API virtual FSceneOutlinerDragValidationInfo ValidateDrop(const ISceneOutlinerTreeItem& DropTarget, const FSceneOutlinerDragDropPayload& Payload) const override;
+	UE_API virtual void OnDrop(ISceneOutlinerTreeItem& DropTarget, const FSceneOutlinerDragDropPayload& Payload, const FSceneOutlinerDragValidationInfo& ValidationInfo) const override;
+	UE_API virtual TSet<FName> GetPropertiesThatRequireListRefresh() const override;
+	UE_API virtual void OnContextMenuContextCreated(FToolMenuContext& Context) const override;
 
 	virtual bool GetShowTransientObjects() const override
 	{
@@ -72,5 +74,7 @@ public:
 
 private:
 	/** Given a tree item, get the object it represents */
-	UObject* GetObjectForTreeItem(const ISceneOutlinerTreeItem& TreeItem) const;
+	UE_API UObject* GetObjectForTreeItem(const ISceneOutlinerTreeItem& TreeItem) const;
 };
+
+#undef UE_API
diff --git a/Engine/Plugins/Editor/ColorGrading/Source/ColorGradingEditor/Public/ColorGradingEditorDataModel.h b/Engine/Plugins/Editor/ColorGrading/Source/ColorGradingEditor/Public/ColorGradingEditorDataModel.h
--- a/Engine/Plugins/Editor/ColorGrading/Source/ColorGradingEditor/Public/ColorGradingEditorDataModel.h
+++ b/Engine/Plugins/Editor/ColorGrading/Source/ColorGradingEditor/Public/ColorGradingEditorDataModel.h
@@ -6,6 +6,8 @@
 #include "IDetailPropertyRow.h"
 #include "PropertyEditorDelegates.h"
 
+#define UE_API COLORGRADINGEDITOR_API
+
 class FResetToDefaultOverride;
 class IPropertyHandle;
 class IPropertyRowGenerator;
@@ -13,7 +15,7 @@ class SWidget;
 struct FColorGradingPanelState;
 
 /** Interface that allows color grading data models to be generated for specific UObjects */
-class COLORGRADINGEDITOR_API IColorGradingEditorDataModelGenerator : public TSharedFromThis<IColorGradingEditorDataModelGenerator>
+class IColorGradingEditorDataModelGenerator : public TSharedFromThis<IColorGradingEditorDataModelGenerator>
 {
 public:
 	virtual ~IColorGradingEditorDataModelGenerator() {}
@@ -31,7 +33,7 @@ DECLARE_DELEGATE_RetVal(TSharedRef<IColorGradingEditorDataModelGenerator>, FGetD
  * A data model that stores the color grading properties from a UObject, which allows UObjects with widely varying color grading
  * interfaces to be used with the color grading panel.
  */
-class COLORGRADINGEDITOR_API FColorGradingEditorDataModel : public TSharedFromThis<FColorGradingEditorDataModel>
+class FColorGradingEditorDataModel : public TSharedFromThis<FColorGradingEditorDataModel>
 {
 public:
 	/** Registers a new data model generator used to populate a color grading data model for the specified class */
@@ -44,7 +46,7 @@ public:
 
 private:
 	/** A list of data model generators, indexed by class name, that have been registered */
-	static TMap<TWeakObjectPtr<UClass>, FGetDetailsDataModelGenerator> RegisteredDataModelGenerators;
+	static UE_API TMap<TWeakObjectPtr<UClass>, FGetDetailsDataModelGenerator> RegisteredDataModelGenerators;
 
 public:
 	/** Stores the property handles for the color properties for a single color grading element (Shadows, Midtones, etc.) */
@@ -118,37 +120,37 @@ public:
 	DECLARE_MULTICAST_DELEGATE_TwoParams(FOnColorGradingGroupRenamed, int32, const FText&);
 
 public:
-	FColorGradingEditorDataModel();
+	UE_API FColorGradingEditorDataModel();
 
 	/** Gets the property row generator that was used to generate this data model */
 	TSharedRef<IPropertyRowGenerator> GetPropertyRowGenerator() const { return PropertyRowGenerator.ToSharedRef(); }
 
 	/** Gets the list of objects from which the color grading properties were generated from */
-	TArray<TWeakObjectPtr<UObject>> GetObjects() const;
+	UE_API TArray<TWeakObjectPtr<UObject>> GetObjects() const;
 
 	/** Sets the objects from which the color grading properties will be generated from */
-	void SetObjects(const TArray<UObject*>& InObjects);
+	UE_API void SetObjects(const TArray<UObject*>& InObjects);
 
 	/** Gets whether the data model has objects of the specified class */
-	bool HasObjectOfType(const UClass* InClass) const;
+	UE_API bool HasObjectOfType(const UClass* InClass) const;
 
 	/** Resets the data model to be empty */
-	void Reset();
+	UE_API void Reset();
 
 	/** Adds the state of the data model to the specified drawer state */
-	void GetPanelState(FColorGradingPanelState& OutPanelState);
+	UE_API void GetPanelState(FColorGradingPanelState& OutPanelState);
 
 	/** Sets the state of the data model from the specified drawer state */
-	void SetPanelState(const FColorGradingPanelState& InPanelState);
+	UE_API void SetPanelState(const FColorGradingPanelState& InPanelState);
 
 	/** Gets the index of the selected color grading group */
 	int32 GetSelectedColorGradingGroupIndex() const { return SelectedColorGradingGroupIndex; }
 
 	/** Gets a pointer to the selected color grading group, or null if there isn't a valid selection */
-	FColorGradingGroup* GetSelectedColorGradingGroup();
+	UE_API FColorGradingGroup* GetSelectedColorGradingGroup();
 
 	/** Sets the selected color grading group */
-	void SetSelectedColorGradingGroup(int32 InColorGradingGroupIndex);
+	UE_API void SetSelectedColorGradingGroup(int32 InColorGradingGroupIndex);
 
 	/** Gets the delegate that is raised when the data model is generated */
 	FOnDataModelGenerated& OnDataModelGenerated() { return OnDataModelGeneratedDelegate; }
@@ -160,10 +162,10 @@ public:
 	int32 GetSelectedColorGradingElementIndex() const { return SelectedColorGradingElementIndex; }
 
 	/** Gets a pointer to the selected color grading element, or null if there isn't a valid selection */
-	FColorGradingElement* GetSelectedColorGradingElement();
+	UE_API FColorGradingElement* GetSelectedColorGradingElement();
 
 	/** Sets the selected color grading element */
-	void SetSelectedColorGradingElement(int32 InColorGradingElementIndex);
+	UE_API void SetSelectedColorGradingElement(int32 InColorGradingElementIndex);
 
 	/** Gets the delegate that is raised when the selected color grading element has changed */
 	FOnColorGradingSelectionChanged& OnColorGradingElementSelectionChanged() { return OnColorGradingElementSelectionChangedDelegate; }
@@ -176,13 +178,13 @@ public:
 
 private:
 	/** Initializes any needed data model generators for the specified class */
-	void InitializeDataModelGenerator(UClass* InClass);
+	UE_API void InitializeDataModelGenerator(UClass* InClass);
 
 	/** Attempts to find an instance of a data model generator that can be used for the specified class */
-	TSharedPtr<IColorGradingEditorDataModelGenerator> GetDataModelGenerator(UClass* InClass) const;
+	UE_API TSharedPtr<IColorGradingEditorDataModelGenerator> GetDataModelGenerator(UClass* InClass) const;
 
 	/** Callback that is raised when the property row generator the data model is generated from has been refreshed */
-	void OnPropertyRowGeneratorRefreshed();
+	UE_API void OnPropertyRowGeneratorRefreshed();
 
 public:
 	/** A list of color grading groups in the color grading data model */
@@ -221,4 +223,6 @@ private:
 
 	/** Delegate that is raised when a color grading group is being renamed */
 	FOnColorGradingGroupRenamed OnColorGradingGroupRenamedDelegate;
-};
\ No newline at end of file
+};
+
+#undef UE_API
diff --git a/Engine/Plugins/Editor/ColorGrading/Source/ColorGradingEditor/Public/ColorGradingMixerContextObject.h b/Engine/Plugins/Editor/ColorGrading/Source/ColorGradingEditor/Public/ColorGradingMixerContextObject.h
--- a/Engine/Plugins/Editor/ColorGrading/Source/ColorGradingEditor/Public/ColorGradingMixerContextObject.h
+++ b/Engine/Plugins/Editor/ColorGrading/Source/ColorGradingEditor/Public/ColorGradingMixerContextObject.h
@@ -7,8 +7,8 @@
 class UToolMenu;
 
 /** Empty class which can be added to an FToolMenuContext to indicate that the menu originates from the Color Grading panel. */
-UCLASS()
-class COLORGRADINGEDITOR_API UColorGradingMixerContextObject : public UObject
+UCLASS(MinimalAPI)
+class UColorGradingMixerContextObject : public UObject
 {
 	GENERATED_BODY()
 };
diff --git a/Engine/Plugins/Editor/ColorGrading/Source/ColorGradingEditor/Public/ColorGradingMixerObjectFilterRegistry.h b/Engine/Plugins/Editor/ColorGrading/Source/ColorGradingEditor/Public/ColorGradingMixerObjectFilterRegistry.h
--- a/Engine/Plugins/Editor/ColorGrading/Source/ColorGradingEditor/Public/ColorGradingMixerObjectFilterRegistry.h
+++ b/Engine/Plugins/Editor/ColorGrading/Source/ColorGradingEditor/Public/ColorGradingMixerObjectFilterRegistry.h
@@ -7,10 +7,12 @@
 #include "SceneOutlinerDragDrop.h"
 #include "Templates/SubclassOf.h"
 
+#define UE_API COLORGRADINGEDITOR_API
+
 class UClass;
 
 /** Interface to configure how an object is displayed in the Color Grading panel's outliner */
-struct COLORGRADINGEDITOR_API IColorGradingMixerObjectHierarchyConfig : public TSharedFromThis<IColorGradingMixerObjectHierarchyConfig>
+struct IColorGradingMixerObjectHierarchyConfig : public TSharedFromThis<IColorGradingMixerObjectHierarchyConfig>
 {
 public:
 	virtual ~IColorGradingMixerObjectHierarchyConfig() {}
@@ -49,7 +51,7 @@ DECLARE_DELEGATE_RetVal(TSharedRef<IColorGradingMixerObjectHierarchyConfig>, FGe
 /**
  * Contains functions for controlling which classes can be managed from the Color Grading panel's ObjectMixer-based hierarchy panel
  */
-struct COLORGRADINGEDITOR_API FColorGradingMixerObjectFilterRegistry
+struct FColorGradingMixerObjectFilterRegistry
 {
 public:
 
@@ -81,18 +83,20 @@ public:
 	static const TSet<TSubclassOf<AActor>>& GetActorClassesToPlace() { return ActorClassesToPlace; }
 
 	/** Get the hierarchy configuration for a class, or null if none was provided */
-	static const IColorGradingMixerObjectHierarchyConfig* GetClassHierarchyConfig(TSubclassOf<AActor> Class);
+	static UE_API const IColorGradingMixerObjectHierarchyConfig* GetClassHierarchyConfig(TSubclassOf<AActor> Class);
 
 private:
 	/** Set of classes that can be seen in the object panel */
-	inline static TSet<UClass*> ObjectClassesToFilter;
+	UE_API inline static TSet<UClass*> ObjectClassesToFilter;
 
 	/** Set of classes that can be placed from the object panel */
-	inline static TSet<TSubclassOf<AActor>> ActorClassesToPlace;
+	UE_API inline static TSet<TSubclassOf<AActor>> ActorClassesToPlace;
 
 	/** Map from object class to hierarchy configuration, if provided */
-	inline static TMap<TWeakObjectPtr<UClass>, FGetObjectHierarchyConfig> HierarchyConfigs;
+	UE_API inline static TMap<TWeakObjectPtr<UClass>, FGetObjectHierarchyConfig> HierarchyConfigs;
 
 	/** Map from object class to hierarchy configuration instance. If null, none exist for the class */
-	inline static TMap<TWeakObjectPtr<UClass>, TSharedPtr<IColorGradingMixerObjectHierarchyConfig>> HierarchyConfigInstances;
+	UE_API inline static TMap<TWeakObjectPtr<UClass>, TSharedPtr<IColorGradingMixerObjectHierarchyConfig>> HierarchyConfigInstances;
 };
+
+#undef UE_API
diff --git a/Engine/Plugins/Editor/ColorGrading/Source/ColorGradingEditor/Public/ColorGradingPanelState.h b/Engine/Plugins/Editor/ColorGrading/Source/ColorGradingEditor/Public/ColorGradingPanelState.h
--- a/Engine/Plugins/Editor/ColorGrading/Source/ColorGradingEditor/Public/ColorGradingPanelState.h
+++ b/Engine/Plugins/Editor/ColorGrading/Source/ColorGradingEditor/Public/ColorGradingPanelState.h
@@ -10,7 +10,7 @@
  * Stores the state of the color grading panel UI that can be reloaded in cases where the panel or any of its elements
  * are reloaded (such as when the containing drawer is reopened or docked)
  */
-struct COLORGRADINGEDITOR_API FColorGradingPanelState
+struct FColorGradingPanelState
 {
 	/** The objects that are selected in the list */
 	TArray<TWeakObjectPtr<UObject>> SelectedObjects;
@@ -33,4 +33,4 @@ struct COLORGRADINGEDITOR_API FColorGradingPanelState
 
 	/** Indicates which subsections were selected for each section in the details panel */
 	TArray<int32> SelectedDetailsSubsections;
-};
\ No newline at end of file
+};
diff --git a/Engine/Plugins/Editor/ColorGrading/Source/ColorGradingEditor/Public/DetailView/SColorGradingDetailView.h b/Engine/Plugins/Editor/ColorGrading/Source/ColorGradingEditor/Public/DetailView/SColorGradingDetailView.h
--- a/Engine/Plugins/Editor/ColorGrading/Source/ColorGradingEditor/Public/DetailView/SColorGradingDetailView.h
+++ b/Engine/Plugins/Editor/ColorGrading/Source/ColorGradingEditor/Public/DetailView/SColorGradingDetailView.h
@@ -9,6 +9,8 @@
 #include "Widgets/SCompoundWidget.h"
 #include "Widgets/Views/STableRow.h"
 
+#define UE_API COLORGRADINGEDITOR_API
+
 class IDetailTreeNode;
 class IPropertyRowGenerator;
 class IPropertyHandle;
@@ -28,7 +30,7 @@ DECLARE_DELEGATE_RetVal_OneParam(bool, FOnFilterDetailTreeNode, const TSharedRef
  * Using an existing property row generator allows the detail view to display an object's properties much faster than the ordinary SDetailView, which
  * has to regenerate a new property node tree every time the object being displayed is changed
  */
-class COLORGRADINGEDITOR_API SColorGradingDetailView : public SCompoundWidget
+class SColorGradingDetailView : public SCompoundWidget
 {
 public:
 	SLATE_BEGIN_ARGS(SColorGradingDetailView) { }
@@ -36,48 +38,48 @@ public:
 		SLATE_EVENT(FOnFilterDetailTreeNode, OnFilterDetailTreeNode)
 	SLATE_END_ARGS()
 
-	void Construct(const FArguments& InArgs);
+	UE_API void Construct(const FArguments& InArgs);
 
 	//~ SWidget interface
-	virtual void Tick(const FGeometry& AllottedGeometry, const double InCurrentTime, const float InDeltaTime) override;
+	UE_API virtual void Tick(const FGeometry& AllottedGeometry, const double InCurrentTime, const float InDeltaTime) override;
 	//~ End SWidget interface
 
 	/** Regenerates this widget based on the current state of its property row generator source */
-	void Refresh();
+	UE_API void Refresh();
 
 	/** Saves the expansion state of all properties being displayed in this detail view to the user's config file */
-	void SaveExpandedItems();
+	UE_API void SaveExpandedItems();
 
 	/** Restores the expansion state of all properties being displayed in this detail view from the user's config file */
-	void RestoreExpandedItems();
+	UE_API void RestoreExpandedItems();
 
 private:
 	/** Updates the detail tree using the current state of the property row generator source */
-	void UpdateTreeNodes();
+	UE_API void UpdateTreeNodes();
 
 	/** Updates the expansion state of the specified tree item using the stored expansion state configuration */
-	void UpdateExpansionState(const TSharedRef<FColorGradingDetailTreeItem> InTreeItem);
+	UE_API void UpdateExpansionState(const TSharedRef<FColorGradingDetailTreeItem> InTreeItem);
 
 	/** Sets the expansion state of the specified tree item, and optionally recursively sets the expansion state of its children */
-	void SetNodeExpansionState(TSharedRef<FColorGradingDetailTreeItem> InTreeNode, bool bIsItemExpanded, bool bRecursive);
+	UE_API void SetNodeExpansionState(TSharedRef<FColorGradingDetailTreeItem> InTreeNode, bool bIsItemExpanded, bool bRecursive);
 
 	/** Generates a table row widget for the specified tree item */
-	TSharedRef<ITableRow> GenerateNodeRow(TSharedRef<FColorGradingDetailTreeItem> InTreeNode, const TSharedRef<STableViewBase>& OwnerTable);
+	UE_API TSharedRef<ITableRow> GenerateNodeRow(TSharedRef<FColorGradingDetailTreeItem> InTreeNode, const TSharedRef<STableViewBase>& OwnerTable);
 
 	/** Gets a list of child tree items for the specified tree item */
-	void GetChildrenForNode(TSharedRef<FColorGradingDetailTreeItem> InTreeNode, TArray<TSharedRef<FColorGradingDetailTreeItem>>& OutChildren);
+	UE_API void GetChildrenForNode(TSharedRef<FColorGradingDetailTreeItem> InTreeNode, TArray<TSharedRef<FColorGradingDetailTreeItem>>& OutChildren);
 
 	/** Raised when the underlying tree widget is setting the expansion state of the specified tree item recursively */
-	void OnSetExpansionRecursive(TSharedRef<FColorGradingDetailTreeItem> InTreeNode, bool bIsItemExpanded);
+	UE_API void OnSetExpansionRecursive(TSharedRef<FColorGradingDetailTreeItem> InTreeNode, bool bIsItemExpanded);
 
 	/** Raised when the underlying tree widget is setting the expansion state of the specified tree item */
-	void OnExpansionChanged(TSharedRef<FColorGradingDetailTreeItem> InTreeNode, bool bIsItemExpanded);
+	UE_API void OnExpansionChanged(TSharedRef<FColorGradingDetailTreeItem> InTreeNode, bool bIsItemExpanded);
 
 	/** Raised when the underlying tree widget is releasing the specified table row */
-	void OnRowReleased(const TSharedRef<ITableRow>& TableRow);
+	UE_API void OnRowReleased(const TSharedRef<ITableRow>& TableRow);
 
 	/** Gets the visibility of the scrollbar for the detail view */
-	EVisibility GetScrollBarVisibility() const;
+	UE_API EVisibility GetScrollBarVisibility() const;
 
 private:
 	typedef STreeView<TSharedRef<FColorGradingDetailTreeItem>> SDetailTree;
@@ -102,4 +104,6 @@ private:
 
 	/** Delegate used to filter or process the detail tree nodes that are displayed in the detail view */
 	FOnFilterDetailTreeNode OnFilterDetailTreeNode;
-};
\ No newline at end of file
+};
+
+#undef UE_API
diff --git a/Engine/Plugins/Editor/ColorGrading/Source/ColorGradingEditor/Public/IColorGradingEditor.h b/Engine/Plugins/Editor/ColorGrading/Source/ColorGradingEditor/Public/IColorGradingEditor.h
--- a/Engine/Plugins/Editor/ColorGrading/Source/ColorGradingEditor/Public/IColorGradingEditor.h
+++ b/Engine/Plugins/Editor/ColorGrading/Source/ColorGradingEditor/Public/IColorGradingEditor.h
@@ -9,7 +9,7 @@
 /**
  * Display Cluster Color Grading module interface
  */
-class COLORGRADINGEDITOR_API IColorGradingEditor : public IModuleInterface
+class IColorGradingEditor : public IModuleInterface
 {
 public:
 	static constexpr const TCHAR* ModuleName = TEXT("ColorGradingEditor");
@@ -40,4 +40,4 @@ public:
 	{
 		return FModuleManager::Get().IsModuleLoaded(ModuleName);
 	}
-};
\ No newline at end of file
+};
diff --git a/Engine/Plugins/Editor/ColorGrading/Source/ColorGradingEditor/Public/SColorGradingPanel.h b/Engine/Plugins/Editor/ColorGrading/Source/ColorGradingEditor/Public/SColorGradingPanel.h
--- a/Engine/Plugins/Editor/ColorGrading/Source/ColorGradingEditor/Public/SColorGradingPanel.h
+++ b/Engine/Plugins/Editor/ColorGrading/Source/ColorGradingEditor/Public/SColorGradingPanel.h
@@ -10,6 +10,8 @@
 
 #include "ColorGradingPanelState.h"
 
+#define UE_API COLORGRADINGEDITOR_API
+
 class FColorGradingEditorDataModel;
 class FObjectMixerEditorList;
 class SColorGradingColorWheelPanel;
@@ -20,10 +22,10 @@ class UWorld;
 using FColorGradingActorFilter = TFunction<bool(const AActor*)>;
 
 /** Main panel of a color grading drawer widget, which displays color wheels or selected object details */
-class COLORGRADINGEDITOR_API SColorGradingPanel : public SCompoundWidget, public FEditorUndoClient
+class SColorGradingPanel : public SCompoundWidget, public FEditorUndoClient
 {
 public:
-	~SColorGradingPanel();
+	UE_API ~SColorGradingPanel();
 
 	SLATE_BEGIN_ARGS(SColorGradingPanel)
 		: _IsInDrawer(false)
@@ -46,81 +48,81 @@ public:
 
 	SLATE_END_ARGS()
 
-	void Construct(const FArguments& InArgs);
+	UE_API void Construct(const FArguments& InArgs);
 
 	/** Refreshes the panel's UI to match the current state of the level */
-	void Refresh();
+	UE_API void Refresh();
 
 	/** Gets the state of the panel UI */
-	void GetPanelState(FColorGradingPanelState& OutPanelState) const;
+	UE_API void GetPanelState(FColorGradingPanelState& OutPanelState) const;
 
 	/** Sets the state of the panel UI */
-	void SetPanelState(const FColorGradingPanelState& InPanelState);
+	UE_API void SetPanelState(const FColorGradingPanelState& InPanelState);
 
 	/**
 	 * Set the list of selected objects, updating state and data model as appropriate.
 	 * If ControlledObjects is provided, only the objects in that array will be used to display the color grading wheels.
 	 * Otherwise, the controlled objects will use the contents of SelectedObjects.
 	 */
-	void SetSelectedObjects(const TArray<UObject*>& SelectedObjects, const TArray<UObject*>* ControlledObjects = nullptr);
+	UE_API void SetSelectedObjects(const TArray<UObject*>& SelectedObjects, const TArray<UObject*>* ControlledObjects = nullptr);
 
 private:
 	/** Creates the button used to dock the drawer in the operator panel */
-	TSharedRef<SWidget> CreateDockInLayoutButton();
+	UE_API TSharedRef<SWidget> CreateDockInLayoutButton();
 
 	/** Get the world currently being edited */
-	UWorld* GetWorld();
+	UE_API UWorld* GetWorld();
 
 	/** Refreshes the object list, filling it with the current color gradable objects from the root actor and world */
-	void RefreshColorGradingList();
+	UE_API void RefreshColorGradingList();
 
 	/** Fills the color grading group toolbar using the color grading data model */
-	void FillColorGradingGroupToolBar();
+	UE_API void FillColorGradingGroupToolBar();
 
 	/** Gets the visibility state of the color grading group toolbar */
-	EVisibility GetColorGradingGroupToolBarVisibility() const;
+	UE_API EVisibility GetColorGradingGroupToolBarVisibility() const;
 
 	/** Gets whether the color grading group at the specified index is currently selected */
-	ECheckBoxState IsColorGradingGroupSelected(int32 GroupIndex) const;
+	UE_API ECheckBoxState IsColorGradingGroupSelected(int32 GroupIndex) const;
 
 	/** Raised when the user has selected the specified color grading group */
-	void OnColorGradingGroupCheckedChanged(ECheckBoxState State, int32 GroupIndex);
+	UE_API void OnColorGradingGroupCheckedChanged(ECheckBoxState State, int32 GroupIndex);
 
 	/** Gets the display name of the specified color grading group */
-	FText GetColorGradingGroupDisplayName(int32 GroupIndex) const;
+	UE_API FText GetColorGradingGroupDisplayName(int32 GroupIndex) const;
 
 	/** Gets the font of the display name label of the specified color grading group */
-	FSlateFontInfo GetColorGradingGroupDisplayNameFont(int32 GroupIndex) const;
+	UE_API FSlateFontInfo GetColorGradingGroupDisplayNameFont(int32 GroupIndex) const;
 
 	/** Gets the content for the right click menu for the color grading group */
-	TSharedRef<SWidget> GetColorGradingGroupMenuContent(int32 GroupIndex);
+	UE_API TSharedRef<SWidget> GetColorGradingGroupMenuContent(int32 GroupIndex);
 
 	/** Raised when a color grading group has been deleted by the user */
-	void OnColorGradingGroupDeleted(int32 GroupIndex);
+	UE_API void OnColorGradingGroupDeleted(int32 GroupIndex);
 
 	/** Raised when a rename has been requested on a color grading group */
-	void OnColorGradingGroupRequestRename(int32 GroupIndex);
+	UE_API void OnColorGradingGroupRequestRename(int32 GroupIndex);
 
 	/** Raised when a rename has been committed on a color grading group */
-	void OnColorGradingGroupRenamed(const FText& InText, ETextCommit::Type TextCommitType, int32 GroupIndex);
+	UE_API void OnColorGradingGroupRenamed(const FText& InText, ETextCommit::Type TextCommitType, int32 GroupIndex);
 
 	/** Raised when the color grading data model has been generated */
-	void OnColorGradingDataModelGenerated();
+	UE_API void OnColorGradingDataModelGenerated();
 
 	/** Raised when the "Dock in Layout" button has been clicked */
-	FReply DockInLayout();
+	UE_API FReply DockInLayout();
 
 	/** Raised when the user has selected a new item in any of the drawer's list views */
-	void OnListSelectionChanged(FSceneOutlinerTreeItemPtr TreeItem, ESelectInfo::Type Type);
+	UE_API void OnListSelectionChanged(FSceneOutlinerTreeItemPtr TreeItem, ESelectInfo::Type Type);
 
 	/** Raised when the drawer's outliner has been synchronized with the editor selection */
-	void OnListSelectionSynchronized();
+	UE_API void OnListSelectionSynchronized();
 
 	/** Update the selected items to match the list view */
-	void UpdateSelectionFromList();
+	UE_API void UpdateSelectionFromList();
 
 	/** Given an array of selected tree items, determine which objects will be selected and/or controlled */
-	void GetSelectedAndControlledObjects(const TArray<FSceneOutlinerTreeItemPtr>& InSelectedItems, TArray<UObject*>& OutSelectedObjects, TArray<UObject*>& OutControlledObjects) const;
+	UE_API void GetSelectedAndControlledObjects(const TArray<FSceneOutlinerTreeItemPtr>& InSelectedItems, TArray<UObject*>& OutSelectedObjects, TArray<UObject*>& OutControlledObjects) const;
 
 private:
 	/** Model for the object mixer list used to display the color gradable object hierarchy. */
@@ -156,3 +158,5 @@ private:
 	/** Function used to filter actors before adding them to the object list. */
 	FColorGradingActorFilter ActorFilter;
 };
+
+#undef UE_API
diff --git a/Engine/Plugins/Editor/ConsoleVariablesEditor/Source/ConsoleVariablesEditor/Public/ConsoleVariablesEditorCommandInfo.h b/Engine/Plugins/Editor/ConsoleVariablesEditor/Source/ConsoleVariablesEditor/Public/ConsoleVariablesEditorCommandInfo.h
--- a/Engine/Plugins/Editor/ConsoleVariablesEditor/Source/ConsoleVariablesEditor/Public/ConsoleVariablesEditorCommandInfo.h
+++ b/Engine/Plugins/Editor/ConsoleVariablesEditor/Source/ConsoleVariablesEditor/Public/ConsoleVariablesEditorCommandInfo.h
@@ -9,9 +9,11 @@
 #include "Misc/DateTime.h"
 #include "UObject/ObjectMacros.h"
 
+#define UE_API CONSOLEVARIABLESEDITOR_API
+
 #define LOCTEXT_NAMESPACE "ConsoleVariablesEditor"
 
-struct CONSOLEVARIABLESEDITOR_API FConsoleVariablesEditorCommandInfo
+struct FConsoleVariablesEditorCommandInfo
 {
 	enum class EConsoleObjectType
 	{
@@ -73,50 +75,50 @@ struct CONSOLEVARIABLESEDITOR_API FConsoleVariablesEditorCommandInfo
 	  * @param bSetInSession If true, this CommandInfo's associated variable row will display "Session" in the UI.
 	  * @param bIsInteractiveChange If true, don't print the command or send it to multi-user/concert.
 	 */
-	void ExecuteCommand(
+	UE_API void ExecuteCommand(
 		const FString& NewValueAsString, const bool bShouldTransactInConcert = true,
 		const bool bSetInSession = true, const bool bIsInteractiveChange = false);
 
 	/**
 	 * Prints the current value or state of the command, similar to the text generated in teh output log when entering a command manually.
 	 */
-	void PrintCommandOrVariable();
+	UE_API void PrintCommandOrVariable();
  
 	/** Get a reference to the cached console object. May return nullptr if unregistered. */
-	IConsoleObject* GetConsoleObjectPtr();
+	UE_API IConsoleObject* GetConsoleObjectPtr();
 
 	/** Return the console object as a console variable object if applicable. May return nullptr if unregistered. */
-	IConsoleVariable* GetConsoleVariablePtr();
+	UE_API IConsoleVariable* GetConsoleVariablePtr();
 
 	/**
 	 *Return the console object as a console command object if applicable.
 	 *Does not consider externally parsed console commands, as they have no associated objects.
 	 */
-	IConsoleCommand* GetConsoleCommandPtr();
+	UE_API IConsoleCommand* GetConsoleCommandPtr();
  
-	static UWorld* GetCurrentWorld();
+	static UE_API UWorld* GetCurrentWorld();
 
-	FString GetHelpText();
+	UE_API FString GetHelpText();
  
-	EConsoleVariableFlags GetSource();
+	UE_API EConsoleVariableFlags GetSource();
  
-	void ClearSourceFlags();
+	UE_API void ClearSourceFlags();
 
-	void SetSourceFlag(const EConsoleVariableFlags InSource);
+	UE_API void SetSourceFlag(const EConsoleVariableFlags InSource);
  
-	FText GetSourceAsText();
+	UE_API FText GetSourceAsText();
 
-	static FText ConvertConsoleVariableSetByFlagToText(const EConsoleVariableFlags InFlag);
+	static UE_API FText ConvertConsoleVariableSetByFlagToText(const EConsoleVariableFlags InFlag);
 
 	/**
 	 * If this console object is a variable, get its value as string via OutValueAsString.
 	 * If it's not a variable type (such as a command) then this function will return false.
 	 */
-	bool GetCurrentValueAsString(FString& OutValueAsString);
+	UE_API bool GetCurrentValueAsString(FString& OutValueAsString);
  
-	bool IsCurrentValueDifferentFromInputValue(const FString& InValueToCompare);
+	UE_API bool IsCurrentValueDifferentFromInputValue(const FString& InValueToCompare);
  
-	void OnConsoleVariableChanged(IConsoleVariable* ChangedVariable);
+	UE_API void OnConsoleVariableChanged(IConsoleVariable* ChangedVariable);
 
 	/** The actual string key or name */
 	UPROPERTY()
@@ -153,3 +155,5 @@ struct CONSOLEVARIABLESEDITOR_API FConsoleVariablesEditorCommandInfo
 };
 
 #undef LOCTEXT_NAMESPACE
+
+#undef UE_API
diff --git a/Engine/Plugins/Editor/ConsoleVariablesEditor/Source/ConsoleVariablesEditor/Public/ConsoleVariablesEditorFunctionLibrary.h b/Engine/Plugins/Editor/ConsoleVariablesEditor/Source/ConsoleVariablesEditor/Public/ConsoleVariablesEditorFunctionLibrary.h
--- a/Engine/Plugins/Editor/ConsoleVariablesEditor/Source/ConsoleVariablesEditor/Public/ConsoleVariablesEditorFunctionLibrary.h
+++ b/Engine/Plugins/Editor/ConsoleVariablesEditor/Source/ConsoleVariablesEditor/Public/ConsoleVariablesEditorFunctionLibrary.h
@@ -9,23 +9,25 @@
 
 #include "ConsoleVariablesEditorFunctionLibrary.generated.h"
 
+#define UE_API CONSOLEVARIABLESEDITOR_API
+
 class FString;
 class UConsoleVariablesAsset;
 
 /** An asset used to track collections of console variables that can be recalled and edited using the Console Variables Editor. */
-UCLASS(BlueprintType)
-class CONSOLEVARIABLESEDITOR_API UConsoleVariablesEditorFunctionLibrary : public UBlueprintFunctionLibrary
+UCLASS(MinimalAPI, BlueprintType)
+class UConsoleVariablesEditorFunctionLibrary : public UBlueprintFunctionLibrary
 {
 	GENERATED_BODY()
 public:
 
 	/** Return the currently loaded list of variables in the Console Variables Editor. */
 	UFUNCTION(BlueprintCallable, Category = "Console Variables Editor")
-	static UConsoleVariablesAsset* GetCurrentlyLoadedPreset();
+	static UE_API UConsoleVariablesAsset* GetCurrentlyLoadedPreset();
 
 	/** Loads the given asset in the Console Variables Editor and sets all its variable values. */
 	UFUNCTION(BlueprintCallable, Category = "Console Variables Editor")
-	static void LoadPresetIntoConsoleVariablesEditor(const UConsoleVariablesAsset* InAsset, const EConsoleVariablesEditorPresetImportMode InImportMode = EConsoleVariablesEditorPresetImportMode::AddToExisting);
+	static UE_API void LoadPresetIntoConsoleVariablesEditor(const UConsoleVariablesAsset* InAsset, const EConsoleVariablesEditorPresetImportMode InImportMode = EConsoleVariablesEditorPresetImportMode::AddToExisting);
 
 	/*
 	 * Saves the current list in the Console Variables Editor to the given asset.
@@ -33,7 +35,7 @@ public:
 	 * @return true if successful.
 	 */
 	UFUNCTION(BlueprintCallable, Category = "Console Variables Editor")
-	static bool CopyCurrentListToAsset(UConsoleVariablesAsset* InAsset);
+	static UE_API bool CopyCurrentListToAsset(UConsoleVariablesAsset* InAsset);
 
 	/*
 	 * Adds a validated command to the current preset with its current value.
@@ -41,7 +43,7 @@ public:
 	 * @return true if successful.
 	 */
 	UFUNCTION(BlueprintCallable, Category = "Console Variables Editor")
-	static bool AddValidatedCommandToCurrentPreset(const FString NewCommand);
+	static UE_API bool AddValidatedCommandToCurrentPreset(const FString NewCommand);
 
 	/*
 	 * Removes a command from the current preset if it exists in the saved data.
@@ -49,42 +51,44 @@ public:
 	 * @return true if successful.
 	 */
 	UFUNCTION(BlueprintCallable, Category = "Console Variables Editor")
-	static bool RemoveCommandFromCurrentPreset(const FString NewCommand);
+	static UE_API bool RemoveCommandFromCurrentPreset(const FString NewCommand);
 
 	/** Return an array of strings containing the command names for each command found in the given preset. */
 	UFUNCTION(BlueprintCallable, Category = "Console Variables Editor")
-	static bool GetListOfCommandsFromPreset(const UConsoleVariablesAsset* InAsset, TArray<FString>& OutCommandList);
+	static UE_API bool GetListOfCommandsFromPreset(const UConsoleVariablesAsset* InAsset, TArray<FString>& OutCommandList);
 
 	/** Set a console variable value directly. Returns true if the console object exists. */
 	UFUNCTION(BlueprintCallable, Category = "Console Variables Editor", meta=(DisplayName="Set Console Variable By Name (Float)"))
-	static bool SetConsoleVariableByName_Float(const FString InCommandName, const float InValue);
+	static UE_API bool SetConsoleVariableByName_Float(const FString InCommandName, const float InValue);
 
 	/** Set a console variable value directly. Returns true if the console object exists. */
 	UFUNCTION(BlueprintCallable, Category = "Console Variables Editor", meta=(DisplayName="Set Console Variable By Name (Bool)"))
-	static bool SetConsoleVariableByName_Bool(const FString InCommandName, const bool InValue);
+	static UE_API bool SetConsoleVariableByName_Bool(const FString InCommandName, const bool InValue);
 
 	/** Set a console variable value directly. Returns true if the console object exists. */
 	UFUNCTION(BlueprintCallable, Category = "Console Variables Editor", meta=(DisplayName="Set Console Variable By Name (Int)"))
-	static bool SetConsoleVariableByName_Int(const FString InCommandName, const int32 InValue);
+	static UE_API bool SetConsoleVariableByName_Int(const FString InCommandName, const int32 InValue);
 	
 	/** Set a console variable value directly. Returns true if the console object exists. */
 	UFUNCTION(BlueprintCallable, Category = "Console Variables Editor", meta=(DisplayName="Set Console Variable By Name (String)"))
-	static bool SetConsoleVariableByName_String(const FString InCommandName, const FString InValue);
+	static UE_API bool SetConsoleVariableByName_String(const FString InCommandName, const FString InValue);
 
 	/** Get a console variable's string value directly. Returns true if the console object exists. */
 	UFUNCTION(BlueprintCallable, Category = "Console Variables Editor")
-	static bool GetConsoleVariableStringValue(const FString InCommandName, FString& OutValue);
+	static UE_API bool GetConsoleVariableStringValue(const FString InCommandName, FString& OutValue);
 
 	/** Set a console variable value directly. Returns true if the console object exists. */
 	UFUNCTION(BlueprintCallable, Category = "Console Variables Editor")
-	static bool GetConsoleVariableSourceByName(const FString InCommandName, FString& OutValue);
+	static UE_API bool GetConsoleVariableSourceByName(const FString InCommandName, FString& OutValue);
 
 	/** Return whether the Multi-user sync setting for the current instance of the editor is enabled. */
 	UFUNCTION(BlueprintCallable, Category = "Console Variables Editor")
-	static bool GetEnableMultiUserCVarSync();
+	static UE_API bool GetEnableMultiUserCVarSync();
 	
 	/** Enable or disable the Multi-user sync setting for the current instance of the editor. */
 	UFUNCTION(BlueprintCallable, Category = "Console Variables Editor")
-	static void SetEnableMultiUserCVarSync(bool bNewSetting);
+	static UE_API void SetEnableMultiUserCVarSync(bool bNewSetting);
 	
 };
+
+#undef UE_API
diff --git a/Engine/Plugins/Editor/ConsoleVariablesEditor/Source/ConsoleVariablesEditor/Public/ConsoleVariablesEditorProjectSettings.h b/Engine/Plugins/Editor/ConsoleVariablesEditor/Source/ConsoleVariablesEditor/Public/ConsoleVariablesEditorProjectSettings.h
--- a/Engine/Plugins/Editor/ConsoleVariablesEditor/Source/ConsoleVariablesEditor/Public/ConsoleVariablesEditorProjectSettings.h
+++ b/Engine/Plugins/Editor/ConsoleVariablesEditor/Source/ConsoleVariablesEditor/Public/ConsoleVariablesEditorProjectSettings.h
@@ -32,8 +32,8 @@ enum class EConsoleVariablesEditorPresetImportMode : uint8
 	UseDefault UMETA(Hidden)
 };
 
-UCLASS(config = Engine, defaultconfig)
-class CONSOLEVARIABLESEDITOR_API UConsoleVariablesEditorProjectSettings : public UObject
+UCLASS(MinimalAPI, config = Engine, defaultconfig)
+class UConsoleVariablesEditorProjectSettings : public UObject
 {
 	GENERATED_BODY()
 public:
diff --git a/Engine/Plugins/Editor/ConsoleVariablesEditor/Source/ConsoleVariablesEditor/Public/Views/Widgets/SConsoleVariablesEditorCustomConsoleInputBox.h b/Engine/Plugins/Editor/ConsoleVariablesEditor/Source/ConsoleVariablesEditor/Public/Views/Widgets/SConsoleVariablesEditorCustomConsoleInputBox.h
--- a/Engine/Plugins/Editor/ConsoleVariablesEditor/Source/ConsoleVariablesEditor/Public/Views/Widgets/SConsoleVariablesEditorCustomConsoleInputBox.h
+++ b/Engine/Plugins/Editor/ConsoleVariablesEditor/Source/ConsoleVariablesEditor/Public/Views/Widgets/SConsoleVariablesEditorCustomConsoleInputBox.h
@@ -6,9 +6,11 @@
 #include "Widgets/Input/SMenuAnchor.h"
 #include "Widgets/Views/SListView.h"
 
+#define UE_API CONSOLEVARIABLESEDITOR_API
+
 class SEditableTextBox;
 
-class CONSOLEVARIABLESEDITOR_API SConsoleVariablesEditorCustomConsoleInputBox : public SCompoundWidget
+class SConsoleVariablesEditorCustomConsoleInputBox : public SCompoundWidget
 {
 public:
 	DECLARE_DELEGATE_OneParam(FOnTextCommitted, const FText&);
@@ -98,32 +100,32 @@ public:
 		FText SuggestionsHighlight;
 	};
 
-	void Construct(const FArguments& InArgs);
+	UE_API void Construct(const FArguments& InArgs);
 
-	virtual ~SConsoleVariablesEditorCustomConsoleInputBox() override;
+	UE_API virtual ~SConsoleVariablesEditorCustomConsoleInputBox() override;
 
 	//~ Begin SWidget Interface
-	virtual FReply OnPreviewKeyDown(const FGeometry& MyGeometry, const FKeyEvent& KeyEvent) override;
+	UE_API virtual FReply OnPreviewKeyDown(const FGeometry& MyGeometry, const FKeyEvent& KeyEvent) override;
 	//~ End SWidget Interface
 
-	bool TakeKeyboardFocus() const;
+	UE_API bool TakeKeyboardFocus() const;
 
-	void OnInputTextChanged(const FText& InText);
+	UE_API void OnInputTextChanged(const FText& InText);
 
-	FReply OnKeyCharHandler(const FGeometry& MyGeometry, const FCharacterEvent& InCharacterEvent) const;
+	UE_API FReply OnKeyCharHandler(const FGeometry& MyGeometry, const FCharacterEvent& InCharacterEvent) const;
 
-	FReply OnKeyDownHandler(const FGeometry& MyGeometry, const FKeyEvent& KeyPressed);
+	UE_API FReply OnKeyDownHandler(const FGeometry& MyGeometry, const FKeyEvent& KeyPressed);
 	
-	void SetSuggestions(TArray<FString>& Elements, FText Highlight);
+	UE_API void SetSuggestions(TArray<FString>& Elements, FText Highlight);
 
-	void MarkActiveSuggestion();
+	UE_API void MarkActiveSuggestion();
 
-	void ClearSuggestions();
+	UE_API void ClearSuggestions();
 	
-	void CommitInput();
+	UE_API void CommitInput();
 
 private:
-	FString GetSanitizedCommand(const FString& InCommand) const;
+	UE_API FString GetSanitizedCommand(const FString& InCommand) const;
 
 private:
 	
@@ -152,3 +154,5 @@ private:
 	FOnTextCommitted OnTextCommitted;
 	FOnTextChanged OnTextChanged;
 };
+
+#undef UE_API
diff --git a/Engine/Plugins/Editor/ConsoleVariablesEditor/Source/ConsoleVariablesEditorRuntime/Public/ConsoleVariablesAsset.h b/Engine/Plugins/Editor/ConsoleVariablesEditor/Source/ConsoleVariablesEditorRuntime/Public/ConsoleVariablesAsset.h
--- a/Engine/Plugins/Editor/ConsoleVariablesEditor/Source/ConsoleVariablesEditorRuntime/Public/ConsoleVariablesAsset.h
+++ b/Engine/Plugins/Editor/ConsoleVariablesEditor/Source/ConsoleVariablesEditorRuntime/Public/ConsoleVariablesAsset.h
@@ -8,6 +8,8 @@
 
 #include "ConsoleVariablesAsset.generated.h"
 
+#define UE_API CONSOLEVARIABLESEDITORRUNTIME_API
+
 /** Data that will be serialized with this asset */
 USTRUCT(BlueprintType)
 struct FConsoleVariablesEditorAssetSaveData
@@ -31,15 +33,15 @@ struct FConsoleVariablesEditorAssetSaveData
 };
 
 /** An asset used to track collections of console variables that can be recalled and edited using the Console Variables Editor. */
-UCLASS(BlueprintType)
-class CONSOLEVARIABLESEDITORRUNTIME_API UConsoleVariablesAsset : public UObject, public IMovieSceneConsoleVariableTrackInterface
+UCLASS(MinimalAPI, BlueprintType)
+class UConsoleVariablesAsset : public UObject, public IMovieSceneConsoleVariableTrackInterface
 {
 	GENERATED_BODY()
 public:
 
 	/** Sets a description for this variable collection. */
 	UFUNCTION(BlueprintCallable, Category = "Console Variables Editor")
-	void SetVariableCollectionDescription(const FString& InVariableCollectionDescription);
+	UE_API void SetVariableCollectionDescription(const FString& InVariableCollectionDescription);
 	
 	UFUNCTION(BlueprintPure, Category = "Console Variables Editor")
 	FString GetVariableCollectionDescription() const
@@ -58,23 +60,23 @@ public:
 	 * @param bOnlyIncludeChecked If true, only commands and variables with a Checked checkstate in the Console Variables Editor UI will be included. Otherwise, all will be included. 
 	 */
 	UFUNCTION(BlueprintPure, Category="Console Variables Asset")
-	TArray<FString> GetSavedCommandsAsStringArray(bool bOnlyIncludeChecked = false) const;
+	UE_API TArray<FString> GetSavedCommandsAsStringArray(bool bOnlyIncludeChecked = false) const;
 
 	/** Returns the saved list of console variables as a concatenated comma-separated string. Useful for passing commands and variables to a command line.
 	 * @param bOnlyIncludeChecked If true, only commands and variables with a Checked checkstate in the Console Variables Editor UI will be included. Otherwise, all will be included. 
 	 */
 	UFUNCTION(BlueprintPure, Category="Console Variables Asset")
-	FString GetSavedCommandsAsCommaSeparatedString(bool bOnlyIncludeChecked = false) const;
+	UE_API FString GetSavedCommandsAsCommaSeparatedString(bool bOnlyIncludeChecked = false) const;
 
 	/** Executes all saved commands in this asset, optionally only including those with a Checked checkstate in the UI.
 	 * @param bOnlyIncludeChecked If true, only commands and variables with a Checked checkstate in the Console Variables Editor UI will be included. Otherwise, all will be included. 
 	 */
 	UFUNCTION(BlueprintCallable, Category="Console Variables Asset", meta=(WorldContext="WorldContextObject"))
-	void ExecuteSavedCommands(UObject* WorldContextObject, bool bOnlyIncludeChecked = false) const;
+	UE_API void ExecuteSavedCommands(UObject* WorldContextObject, bool bOnlyIncludeChecked = false) const;
 
 	/** Completely replaces the saved data with new saved data */
 	UFUNCTION(BlueprintCallable, Category="Console Variables Asset")
-	void ReplaceSavedCommands(const TArray<FConsoleVariablesEditorAssetSaveData>& Replacement);
+	UE_API void ReplaceSavedCommands(const TArray<FConsoleVariablesEditorAssetSaveData>& Replacement);
 
 	/** Returns how many console variables are serialized in this asset */
 	UFUNCTION(BlueprintPure, Category="Console Variables Asset")
@@ -85,30 +87,30 @@ public:
 
 	/** Outputs the FConsoleVariablesEditorAssetSaveData matching InCommand. Returns whether a match was found. Case sensitive. */
 	UFUNCTION(BlueprintCallable, Category="Console Variables Asset")
-	bool FindSavedDataByCommandString(
+	UE_API bool FindSavedDataByCommandString(
 		const FString& InCommandString, FConsoleVariablesEditorAssetSaveData& OutValue, const ESearchCase::Type SearchCase = ESearchCase::IgnoreCase) const;
 
 	/** Set the value of a saved console variable if the name matches; add a new console variable to the list if a match is not found. */
 	UFUNCTION(BlueprintCallable, Category="Console Variables Asset")
-	void AddOrSetConsoleObjectSavedData(const FConsoleVariablesEditorAssetSaveData& InData);
+	UE_API void AddOrSetConsoleObjectSavedData(const FConsoleVariablesEditorAssetSaveData& InData);
 
 	/** Returns true if the element was found and successfully removed. */
 	UFUNCTION(BlueprintCallable, Category="Console Variables Asset")
-	bool RemoveConsoleVariable(const FString& InCommandString);
+	UE_API bool RemoveConsoleVariable(const FString& InCommandString);
 
 	/** Copy list of variables from input asset to this asset */
 	UFUNCTION(BlueprintCallable, Category="Console Variables Asset")
-	void CopyFrom(const UConsoleVariablesAsset* InAssetToCopy);
+	UE_API void CopyFrom(const UConsoleVariablesAsset* InAssetToCopy);
 
 	/** Add list of variables from input asset to this asset's list */
 	UFUNCTION(BlueprintCallable, Category = "Console Variables Asset")
-	void AddFrom(const UConsoleVariablesAsset* InAssetToCopy);
+	UE_API void AddFrom(const UConsoleVariablesAsset* InAssetToCopy);
 	
 private:
 
 	//~ IMovieSceneConsoleVariableTrackInterface overrides
-	virtual void GetConsoleVariablesForTrack(bool bOnlyIncludeChecked, TArray<TPair<FString, FString>>& OutVariables) override;
-	virtual bool IsConsoleVariableEnabled(const FString& ConsoleVariableName) override;
+	UE_API virtual void GetConsoleVariablesForTrack(bool bOnlyIncludeChecked, TArray<TPair<FString, FString>>& OutVariables) override;
+	UE_API virtual bool IsConsoleVariableEnabled(const FString& ConsoleVariableName) override;
 
 private:
 
@@ -120,3 +122,5 @@ private:
 	UPROPERTY()
 	TArray<FConsoleVariablesEditorAssetSaveData> SavedCommands;
 };
+
+#undef UE_API
diff --git a/Engine/Plugins/Editor/ContentBrowser/ContentBrowserAliasDataSource/Source/ContentBrowserAliasDataSource/Public/ContentBrowserAliasDataSource.h b/Engine/Plugins/Editor/ContentBrowser/ContentBrowserAliasDataSource/Source/ContentBrowserAliasDataSource/Public/ContentBrowserAliasDataSource.h
--- a/Engine/Plugins/Editor/ContentBrowser/ContentBrowserAliasDataSource/Source/ContentBrowserAliasDataSource/Public/ContentBrowserAliasDataSource.h
+++ b/Engine/Plugins/Editor/ContentBrowser/ContentBrowserAliasDataSource/Source/ContentBrowserAliasDataSource/Public/ContentBrowserAliasDataSource.h
@@ -12,6 +12,8 @@
 #include "UObject/Package.h"
 #include "ContentBrowserAliasDataSource.generated.h"
 
+#define UE_API CONTENTBROWSERALIASDATASOURCE_API
+
 class IAssetRegistry;
 struct FPropertyChangedEvent;
 
@@ -23,7 +25,7 @@ struct FContentBrowserCompiledAssetDataFilter;
 /** A unique alias is a pair of SourceObjectPath:AliasPath, eg /Game/MyAsset.MyAsset:/Game/SomeFolder/MyAlias */
 typedef TPair<FSoftObjectPath, FName> FContentBrowserUniqueAlias;
 
-class CONTENTBROWSERALIASDATASOURCE_API FContentBrowserAliasItemDataPayload : public FContentBrowserAssetFileItemDataPayload
+class FContentBrowserAliasItemDataPayload : public FContentBrowserAssetFileItemDataPayload
 {
 public:
 	FContentBrowserAliasItemDataPayload(const FAssetData& InAssetData, const FContentBrowserUniqueAlias& InAlias)
@@ -42,77 +44,77 @@ public:
  * or manually managed by calling AddAlias and RemoveAlias. ReconcileAliasesForAsset is provided as a helper function to automatically update new/removed aliases for an existing asset.
  *
  */
-UCLASS()
-class CONTENTBROWSERALIASDATASOURCE_API UContentBrowserAliasDataSource : public UContentBrowserDataSource
+UCLASS(MinimalAPI)
+class UContentBrowserAliasDataSource : public UContentBrowserDataSource
 {
 	GENERATED_BODY()
 
 public:
 	/** The metadata tag to set for the AliasDataSource to automatically create aliases for an asset */
-	static FName AliasTagName;
+	static UE_API FName AliasTagName;
 
 	// ~ Begin UContentBrowserDataSource interface
-	void Initialize(const bool InAutoRegister = true);
-	virtual void Shutdown() override;
+	UE_API void Initialize(const bool InAutoRegister = true);
+	UE_API virtual void Shutdown() override;
 
-	virtual void CompileFilter(const FName InPath, const FContentBrowserDataFilter& InFilter, FContentBrowserDataCompiledFilter& OutCompiledFilter) override;
-	virtual void EnumerateItemsMatchingFilter(const FContentBrowserDataCompiledFilter& InFilter, TFunctionRef<bool(FContentBrowserItemData&&)> InCallback) override;
-	virtual void EnumerateItemsAtPath(const FName InPath, const EContentBrowserItemTypeFilter InItemTypeFilter, TFunctionRef<bool(FContentBrowserItemData&&)> InCallback) override;
+	UE_API virtual void CompileFilter(const FName InPath, const FContentBrowserDataFilter& InFilter, FContentBrowserDataCompiledFilter& OutCompiledFilter) override;
+	UE_API virtual void EnumerateItemsMatchingFilter(const FContentBrowserDataCompiledFilter& InFilter, TFunctionRef<bool(FContentBrowserItemData&&)> InCallback) override;
+	UE_API virtual void EnumerateItemsAtPath(const FName InPath, const EContentBrowserItemTypeFilter InItemTypeFilter, TFunctionRef<bool(FContentBrowserItemData&&)> InCallback) override;
 
-	virtual bool IsFolderVisible(const FName Path, const EContentBrowserIsFolderVisibleFlags Flags, const FContentBrowserFolderContentsFilter& ContentsFilter) override;
+	UE_API virtual bool IsFolderVisible(const FName Path, const EContentBrowserIsFolderVisibleFlags Flags, const FContentBrowserFolderContentsFilter& ContentsFilter) override;
 
-	virtual bool DoesItemPassFilter(const FContentBrowserItemData& InItem, const FContentBrowserDataCompiledFilter& InFilter) override;
+	UE_API virtual bool DoesItemPassFilter(const FContentBrowserItemData& InItem, const FContentBrowserDataCompiledFilter& InFilter) override;
 
 	using UContentBrowserDataSource::GetAliasesForPath;
-	virtual TArray<FContentBrowserItemPath> GetAliasesForPath(const FSoftObjectPath& InInternalPath) const override;
+	UE_API virtual TArray<FContentBrowserItemPath> GetAliasesForPath(const FSoftObjectPath& InInternalPath) const override;
 
-	bool HasAliasesForPath(const FSoftObjectPath& InInternalPath) const;
+	UE_API bool HasAliasesForPath(const FSoftObjectPath& InInternalPath) const;
 
-	virtual bool GetItemAttribute(const FContentBrowserItemData& InItem, const bool InIncludeMetaData, const FName InAttributeKey, FContentBrowserItemDataAttributeValue& OutAttributeValue) override;
-	virtual bool GetItemAttributes(const FContentBrowserItemData& InItem, const bool InIncludeMetaData, FContentBrowserItemDataAttributeValues& OutAttributeValues) override;
-	virtual bool GetItemPhysicalPath(const FContentBrowserItemData& InItem, FString& OutDiskPath) override;
-	virtual bool IsItemDirty(const FContentBrowserItemData& InItem) override;
+	UE_API virtual bool GetItemAttribute(const FContentBrowserItemData& InItem, const bool InIncludeMetaData, const FName InAttributeKey, FContentBrowserItemDataAttributeValue& OutAttributeValue) override;
+	UE_API virtual bool GetItemAttributes(const FContentBrowserItemData& InItem, const bool InIncludeMetaData, FContentBrowserItemDataAttributeValues& OutAttributeValues) override;
+	UE_API virtual bool GetItemPhysicalPath(const FContentBrowserItemData& InItem, FString& OutDiskPath) override;
+	UE_API virtual bool IsItemDirty(const FContentBrowserItemData& InItem) override;
 
-	virtual bool CanEditItem(const FContentBrowserItemData& InItem, FText* OutErrorMsg) override;
-	virtual bool EditItem(const FContentBrowserItemData& InItem) override;
-	virtual bool BulkEditItems(TArrayView<const FContentBrowserItemData> InItems) override;
+	UE_API virtual bool CanEditItem(const FContentBrowserItemData& InItem, FText* OutErrorMsg) override;
+	UE_API virtual bool EditItem(const FContentBrowserItemData& InItem) override;
+	UE_API virtual bool BulkEditItems(TArrayView<const FContentBrowserItemData> InItems) override;
 	
-	virtual bool CanViewItem(const FContentBrowserItemData& InItem, FText* OutErrorMsg) override;
-	virtual bool ViewItem(const FContentBrowserItemData& InItem) override;
-	virtual bool BulkViewItems(TArrayView<const FContentBrowserItemData> InItems) override;
+	UE_API virtual bool CanViewItem(const FContentBrowserItemData& InItem, FText* OutErrorMsg) override;
+	UE_API virtual bool ViewItem(const FContentBrowserItemData& InItem) override;
+	UE_API virtual bool BulkViewItems(TArrayView<const FContentBrowserItemData> InItems) override;
 
-	virtual bool CanPreviewItem(const FContentBrowserItemData& InItem, FText* OutErrorMsg) override;
-	virtual bool PreviewItem(const FContentBrowserItemData& InItem) override;
-	virtual bool BulkPreviewItems(TArrayView<const FContentBrowserItemData> InItems) override;
+	UE_API virtual bool CanPreviewItem(const FContentBrowserItemData& InItem, FText* OutErrorMsg) override;
+	UE_API virtual bool PreviewItem(const FContentBrowserItemData& InItem) override;
+	UE_API virtual bool BulkPreviewItems(TArrayView<const FContentBrowserItemData> InItems) override;
 
-	virtual bool CanSaveItem(const FContentBrowserItemData& InItem, const EContentBrowserItemSaveFlags InSaveFlags, FText* OutErrorMsg) override;
-	virtual bool SaveItem(const FContentBrowserItemData& InItem, const EContentBrowserItemSaveFlags InSaveFlags) override;
-	virtual bool BulkSaveItems(TArrayView<const FContentBrowserItemData> InItems, const EContentBrowserItemSaveFlags InSaveFlags) override;
+	UE_API virtual bool CanSaveItem(const FContentBrowserItemData& InItem, const EContentBrowserItemSaveFlags InSaveFlags, FText* OutErrorMsg) override;
+	UE_API virtual bool SaveItem(const FContentBrowserItemData& InItem, const EContentBrowserItemSaveFlags InSaveFlags) override;
+	UE_API virtual bool BulkSaveItems(TArrayView<const FContentBrowserItemData> InItems, const EContentBrowserItemSaveFlags InSaveFlags) override;
 
-	virtual bool CanPrivatizeItem(const FContentBrowserItemData& InItem, FText* OutErrorMsg) override;
-	virtual bool PrivatizeItem(const FContentBrowserItemData& InItem, const EAssetAccessSpecifier InAssetAccessSpecifier = EAssetAccessSpecifier::Private) override;
-	virtual bool BulkPrivatizeItems(TArrayView<const FContentBrowserItemData> InItems, const EAssetAccessSpecifier InAssetAccessSpecifier = EAssetAccessSpecifier::Private) override;
+	UE_API virtual bool CanPrivatizeItem(const FContentBrowserItemData& InItem, FText* OutErrorMsg) override;
+	UE_API virtual bool PrivatizeItem(const FContentBrowserItemData& InItem, const EAssetAccessSpecifier InAssetAccessSpecifier = EAssetAccessSpecifier::Private) override;
+	UE_API virtual bool BulkPrivatizeItems(TArrayView<const FContentBrowserItemData> InItems, const EAssetAccessSpecifier InAssetAccessSpecifier = EAssetAccessSpecifier::Private) override;
 
-	virtual bool AppendItemReference(const FContentBrowserItemData& InItem, FString& InOutStr) override;
-	virtual bool AppendItemObjectPath(const FContentBrowserItemData& InItem, FString& InOutStr) override;
-	virtual bool AppendItemPackageName(const FContentBrowserItemData& InItem, FString& InOutStr) override;
-	virtual bool UpdateThumbnail(const FContentBrowserItemData& InItem, FAssetThumbnail& InThumbnail) override;
+	UE_API virtual bool AppendItemReference(const FContentBrowserItemData& InItem, FString& InOutStr) override;
+	UE_API virtual bool AppendItemObjectPath(const FContentBrowserItemData& InItem, FString& InOutStr) override;
+	UE_API virtual bool AppendItemPackageName(const FContentBrowserItemData& InItem, FString& InOutStr) override;
+	UE_API virtual bool UpdateThumbnail(const FContentBrowserItemData& InItem, FAssetThumbnail& InThumbnail) override;
 
-	virtual bool TryGetCollectionId(const FContentBrowserItemData& InItem, FSoftObjectPath& OutCollectionId) override;
+	UE_API virtual bool TryGetCollectionId(const FContentBrowserItemData& InItem, FSoftObjectPath& OutCollectionId) override;
 
 	virtual bool GetItemAssetAccessSpecifier(const FContentBrowserItemData& InItem, EAssetAccessSpecifier& OutAssetAccessSpecifier) override;
 	virtual bool CanModifyItemAssetAccessSpecifier(const FContentBrowserItemData& InItem) override;
 
 	// Legacy functions seem necessary for FrontendFilters to work
-	virtual bool Legacy_TryGetPackagePath(const FContentBrowserItemData& InItem, FName& OutPackagePath) override;
-	virtual bool Legacy_TryGetAssetData(const FContentBrowserItemData& InItem, FAssetData& OutAssetData) override;
-	virtual bool Legacy_TryConvertPackagePathToVirtualPath(const FName InPackagePath, FName& OutPath) override;
-	virtual bool Legacy_TryConvertAssetDataToVirtualPath(const FAssetData& InAssetData, const bool InUseFolderPaths, FName& OutPath) override;
+	UE_API virtual bool Legacy_TryGetPackagePath(const FContentBrowserItemData& InItem, FName& OutPackagePath) override;
+	UE_API virtual bool Legacy_TryGetAssetData(const FContentBrowserItemData& InItem, FAssetData& OutAssetData) override;
+	UE_API virtual bool Legacy_TryConvertPackagePathToVirtualPath(const FName InPackagePath, FName& OutPath) override;
+	UE_API virtual bool Legacy_TryConvertAssetDataToVirtualPath(const FAssetData& InAssetData, const bool InUseFolderPaths, FName& OutPath) override;
 	// ~ End UContentBrowserDataSource interface
 
-	virtual void RemoveUnusedCachedFilterData(const FContentBrowserDataFilterCacheIDOwner& IDOwner, TArrayView<const FName> InVirtualPathsInUse, const FContentBrowserDataFilter& DataFilter) override;
+	UE_API virtual void RemoveUnusedCachedFilterData(const FContentBrowserDataFilterCacheIDOwner& IDOwner, TArrayView<const FName> InVirtualPathsInUse, const FContentBrowserDataFilter& DataFilter) override;
 
-	virtual void ClearCachedFilterData(const FContentBrowserDataFilterCacheIDOwner& IDOwner) override;
+	UE_API virtual void ClearCachedFilterData(const FContentBrowserDataFilterCacheIDOwner& IDOwner) override;
 
 	/** 
 	 * Add a list of aliases for a given asset.
@@ -122,52 +124,52 @@ public:
 	 * @param bInIsFromMetaData 			Should only be `true` if the list of aliases came from the `AliasTagName` metadata.
 	 * @param bSkipPrimaryAssetValidation 	Use this if the asset being added is not a primary asset/is a re-director but should still be included in the Content Browser. (e.g. Verse classes)
 	 */
-	void AddAliases(const FAssetData& Asset, const TArray<FName>& Aliases, const bool bInIsFromMetaData = false, const bool bSkipPrimaryAssetValidation = false);
-	void AddAliases(const FAssetData& Asset, const TArray<FContentBrowserLocalizedAlias>& Aliases, const bool bInIsFromMetaData = false, const bool bSkipPrimaryAssetValidation = false);
+	UE_API void AddAliases(const FAssetData& Asset, const TArray<FName>& Aliases, const bool bInIsFromMetaData = false, const bool bSkipPrimaryAssetValidation = false);
+	UE_API void AddAliases(const FAssetData& Asset, const TArray<FContentBrowserLocalizedAlias>& Aliases, const bool bInIsFromMetaData = false, const bool bSkipPrimaryAssetValidation = false);
 	/** Add an alias for a given asset. bInIsFromMetaData should only be true if the alias came from the AliasTagName metadata. */
-	void AddAlias(const FAssetData& Asset, const FName Alias, const bool bInIsFromMetaData = false, const bool bSkipPrimaryAssetValidation = false);
-	void AddAlias(const FAssetData& Asset, const FContentBrowserLocalizedAlias& Alias, const bool bInIsFromMetaData = false, const bool bSkipPrimaryAssetValidation = false);
+	UE_API void AddAlias(const FAssetData& Asset, const FName Alias, const bool bInIsFromMetaData = false, const bool bSkipPrimaryAssetValidation = false);
+	UE_API void AddAlias(const FAssetData& Asset, const FContentBrowserLocalizedAlias& Alias, const bool bInIsFromMetaData = false, const bool bSkipPrimaryAssetValidation = false);
 	/** Remove the given alias from the data source */
-	void RemoveAlias(const FSoftObjectPath& ObjectPath, const FName Alias);
+	UE_API void RemoveAlias(const FSoftObjectPath& ObjectPath, const FName Alias);
 	/** Remove all aliases for the given object */
-	void RemoveAliases(const FSoftObjectPath& ObjectPath);
+	UE_API void RemoveAliases(const FSoftObjectPath& ObjectPath);
 	/** Remove all aliases for the given asset */
 	void RemoveAliases(const FAssetData& Asset) { RemoveAliases(Asset.GetSoftObjectPath()); }
 
 	UE_DEPRECATED(5.1, "FNames containing full asset paths are deprecated, use FSoftObjectPath instead")
-	void RemoveAlias(const FName ObjectPath, const FName Alias);
+	UE_API void RemoveAlias(const FName ObjectPath, const FName Alias);
 	UE_DEPRECATED(5.1, "FNames containing full asset paths are deprecated, use FSoftObjectPath instead")
-	void RemoveAliases(const FName ObjectPath);
+	UE_API void RemoveAliases(const FName ObjectPath);
 
 	/**
 	 * Add a display name override for the given alias folder, eg /MyAliases.
 	 * @note To provide a display name override for an alias itself, use the overloads that take a FContentBrowserLocalizedAlias.
 	 */
-	void AddAliasFolderDisplayName(const FName AliasFolder, const FText& DisplayName);
+	UE_API void AddAliasFolderDisplayName(const FName AliasFolder, const FText& DisplayName);
 	/**
 	 * Remove a display name override for the given alias folder, eg /MyAliases.
 	 */
-	void RemoveAliasFolderDisplayName(const FName AliasFolder);
+	UE_API void RemoveAliasFolderDisplayName(const FName AliasFolder);
 
 	/** When called, removes all aliases and triggers delegate for various systems to re-add aliases */
-	void RebuildAliases();
+	UE_API void RebuildAliases();
 
 	/** Broadcast after RebuildAliases() called to allow systems to re-add aliases */
 	FSimpleMulticastDelegate& OnRebuildAliases() { return RebuildAliasesDelegate; }
 
 	/** Get all aliases from metadata for the given asset, then calls AddAlias or RemoveAlias for every alias that doesn't match the stored data. */
-	void ReconcileAliasesFromMetaData(const FAssetData& Asset);
+	UE_API void ReconcileAliasesFromMetaData(const FAssetData& Asset);
 	/** Calls AddAlias or RemoveAlias for every alias that doesn't match the stored data for the given asset. */
-	void ReconcileAliasesForAsset(const FAssetData& Asset, const TArray<FName>& NewAliases);
-	void ReconcileAliasesForAsset(const FAssetData& Asset, const TArray<FContentBrowserLocalizedAlias>& NewAliases);
+	UE_API void ReconcileAliasesForAsset(const FAssetData& Asset, const TArray<FName>& NewAliases);
+	UE_API void ReconcileAliasesForAsset(const FAssetData& Asset, const TArray<FContentBrowserLocalizedAlias>& NewAliases);
 
 	/** Logs all the content browser aliases */
-	void LogAliases() const;
+	UE_API void LogAliases() const;
 
-	void SetFilterShouldMatchCollectionContent(bool bInFilterShouldMatchCollectionContent);
+	UE_API void SetFilterShouldMatchCollectionContent(bool bInFilterShouldMatchCollectionContent);
 
 protected:
-	virtual void BuildRootPathVirtualTree() override;
+	UE_API virtual void BuildRootPathVirtualTree() override;
 
 private:
 	template <typename AliasType>
@@ -177,20 +179,20 @@ private:
 	template <typename AliasType>
 	void ReconcileAliasesForAssetImpl(const FAssetData& Asset, const TArray<AliasType>& NewAliases);
 
-	void OnAssetAdded(const FAssetData& InAssetData);
-	void OnAssetRemoved(const FAssetData& InAssetData);
-	void OnAssetUpdated(const FAssetData& InAssetData);
-	void OnAssetLoaded(UObject* InAsset);
-	void OnObjectPropertyChanged(UObject* InObject, FPropertyChangedEvent& InPropertyChangedEvent);
+	UE_API void OnAssetAdded(const FAssetData& InAssetData);
+	UE_API void OnAssetRemoved(const FAssetData& InAssetData);
+	UE_API void OnAssetUpdated(const FAssetData& InAssetData);
+	UE_API void OnAssetLoaded(UObject* InAsset);
+	UE_API void OnObjectPropertyChanged(UObject* InObject, FPropertyChangedEvent& InPropertyChangedEvent);
 
 	/** Helper function to remove a folder from the PathTree including all parent folders that are now empty as a result of the removal */
-	void RemoveFoldersRecursively(FStringView LeafFolder);
+	UE_API void RemoveFoldersRecursively(FStringView LeafFolder);
 
-	void UpdateAliasesCachedAssetData(const FAssetData& InAssetData);
-	void MakeItemModifiedUpdate(const FSoftObjectPath& ObjectPath);
+	UE_API void UpdateAliasesCachedAssetData(const FAssetData& InAssetData);
+	UE_API void MakeItemModifiedUpdate(const FSoftObjectPath& ObjectPath);
 
-	FContentBrowserItemData CreateAssetFolderItem(const FName InFolderPath);
-	FContentBrowserItemData CreateAssetFileItem(const FContentBrowserUniqueAlias& Alias);
+	UE_API FContentBrowserItemData CreateAssetFolderItem(const FName InFolderPath);
+	UE_API FContentBrowserItemData CreateAssetFileItem(const FContentBrowserUniqueAlias& Alias);
 
 	IAssetRegistry* AssetRegistry = nullptr;
 	IAssetTools* AssetTools = nullptr;
@@ -231,7 +233,7 @@ private:
 		/** Whether this alias was generated from package metadata or manually through the C++ interface */
 		bool bIsFromMetaData = false;
 	};
-	bool DoesAliasPassFilter(const FAliasData& AliasData, const FContentBrowserCompiledAssetDataFilter& Filter) const;
+	UE_API bool DoesAliasPassFilter(const FAliasData& AliasData, const FContentBrowserCompiledAssetDataFilter& Filter) const;
 
 	/** The full folder hierarchy for all alias paths */
 	FPathTree PathTree;
@@ -254,3 +256,5 @@ private:
 	/** If true, EnumerateItemsMatchingFilter will include aliases of items contained in collections included in the filter */
 	bool bFilterShouldMatchCollectionContent = true;
 };
+
+#undef UE_API
diff --git a/Engine/Plugins/Editor/ContentBrowser/ContentBrowserAssetDataSource/Source/ContentBrowserAssetDataSource/Public/AssetFileContextMenu.h b/Engine/Plugins/Editor/ContentBrowser/ContentBrowserAssetDataSource/Source/ContentBrowserAssetDataSource/Public/AssetFileContextMenu.h
--- a/Engine/Plugins/Editor/ContentBrowser/ContentBrowserAssetDataSource/Source/ContentBrowserAssetDataSource/Public/AssetFileContextMenu.h
+++ b/Engine/Plugins/Editor/ContentBrowser/ContentBrowserAssetDataSource/Source/ContentBrowserAssetDataSource/Public/AssetFileContextMenu.h
@@ -4,19 +4,21 @@
 
 #include "AssetRegistry/AssetData.h"
 
+#define UE_API CONTENTBROWSERASSETDATASOURCE_API
+
 class FReply;
 
 class UToolMenu;
 class SWindow;
 class SWidget;
 
-class CONTENTBROWSERASSETDATASOURCE_API FAssetFileContextMenu : public TSharedFromThis<FAssetFileContextMenu>
+class FAssetFileContextMenu : public TSharedFromThis<FAssetFileContextMenu>
 {
 public:
 	DECLARE_DELEGATE_OneParam(FOnShowAssetsInPathsView, const TArray<FAssetData>& /*AssetsToFind*/);
 
 	/** Makes the context menu widget */
-	void MakeContextMenu(
+	UE_API void MakeContextMenu(
 		UToolMenu* InMenu,
 		const TArray<FAssetData>& InSelectedAssets,
 		const FOnShowAssetsInPathsView& InOnShowAssetsInPathsView
@@ -38,203 +40,203 @@ private:
 
 private:
 	/** Helper to load selected assets and sort them by UClass */
-	void GetSelectedAssetsByClass(TMap<UClass*, TArray<FAssetData> >& OutSelectedAssetsByClass) const;
+	UE_API void GetSelectedAssetsByClass(TMap<UClass*, TArray<FAssetData> >& OutSelectedAssetsByClass) const;
 
 	/** Helper to collect resolved filepaths for all selected assets */
-	void GetSelectedAssetSourceFilePaths(TArray<FString>& OutFilePaths, TArray<FString>& OutUniqueSourceFileLabels, int32 &OutValidSelectedAssetCount) const;
+	UE_API void GetSelectedAssetSourceFilePaths(TArray<FString>& OutFilePaths, TArray<FString>& OutUniqueSourceFileLabels, int32 &OutValidSelectedAssetCount) const;
 
 	/** Handler to check to see if a imported asset actions should be visible in the menu */
-	bool AreImportedAssetActionsVisible() const;
+	UE_API bool AreImportedAssetActionsVisible() const;
 
 	/** Handler to check to see if imported asset actions are allowed */
-	bool CanExecuteImportedAssetActions(const TArray<FString> ResolvedFilePaths) const;
+	UE_API bool CanExecuteImportedAssetActions(const TArray<FString> ResolvedFilePaths) const;
 
 	/** Handler to check to see if reimport asset actions are allowed */
-	bool CanExecuteReimportAssetActions(const TArray<FString> ResolvedFilePaths) const;
+	UE_API bool CanExecuteReimportAssetActions(const TArray<FString> ResolvedFilePaths) const;
 
 	/** Handler for Reimport */
-	void ExecuteReimport(int32 SourceFileIndex = INDEX_NONE);
+	UE_API void ExecuteReimport(int32 SourceFileIndex = INDEX_NONE);
 
-	void ExecuteReimportWithNewFile(int32 SourceFileIndex = INDEX_NONE);
+	UE_API void ExecuteReimportWithNewFile(int32 SourceFileIndex = INDEX_NONE);
 
 	/** Handler for FindInExplorer */
-	void ExecuteFindSourceInExplorer(const TArray<FString> ResolvedFilePaths);
+	UE_API void ExecuteFindSourceInExplorer(const TArray<FString> ResolvedFilePaths);
 
 	/** Handler for OpenInExternalEditor */
-	void ExecuteOpenInExternalEditor(const TArray<FString> ResolvedFilePaths);
+	UE_API void ExecuteOpenInExternalEditor(const TArray<FString> ResolvedFilePaths);
 
 	/** Handler to check to see if a load command is allowed */
-	bool CanExecuteLoad() const;
+	UE_API bool CanExecuteLoad() const;
 
 	/** Handler for Load */
-	void ExecuteLoad();
+	UE_API void ExecuteLoad();
 
 	/** Handler to check to see if a reload command is allowed */
-	bool CanExecuteReload() const;
+	UE_API bool CanExecuteReload() const;
 
 	/** Handler for Reload */
-	void ExecuteReload();
+	UE_API void ExecuteReload();
 
 	/** Is allowed to modify files or folders under this path */
-	bool CanModifyPath(const FString& InPath) const;
+	UE_API bool CanModifyPath(const FString& InPath) const;
 
 	/** Adds options to menu */
-	void AddMenuOptions(UToolMenu* Menu);
+	UE_API void AddMenuOptions(UToolMenu* Menu);
 
 	/** Adds asset type-specific menu options to a menu builder. Returns true if any options were added. */
-	bool AddImportedAssetMenuOptions(UToolMenu* Menu);
+	UE_API bool AddImportedAssetMenuOptions(UToolMenu* Menu);
 	
 	/** Adds common menu options to a menu builder. Returns true if any options were added. */
-	bool AddCommonMenuOptions(UToolMenu* Menu);
+	UE_API bool AddCommonMenuOptions(UToolMenu* Menu);
 
 	/** Adds Asset Actions sub-menu to a menu builder. */
-	void MakeAssetActionsSubMenu(UToolMenu* Menu);
+	UE_API void MakeAssetActionsSubMenu(UToolMenu* Menu);
 
 	/** Handler to check to see if "Asset Actions" are allowed */
-	bool CanExecuteAssetActions() const;
+	UE_API bool CanExecuteAssetActions() const;
 
 	/** Adds Asset Localization sub-menu to a menu builder. */
-	void MakeAssetLocalizationSubMenu(UToolMenu* Menu);
+	UE_API void MakeAssetLocalizationSubMenu(UToolMenu* Menu);
 
 	/** Adds the Create Localized Asset sub-menu to a menu builder. */
-	void MakeCreateLocalizedAssetSubMenu(UToolMenu* Menu, TSet<FSoftObjectPath> InSelectedSourceAssets, TArray<FLocalizedAssetsState> InLocalizedAssetsState);
+	UE_API void MakeCreateLocalizedAssetSubMenu(UToolMenu* Menu, TSet<FSoftObjectPath> InSelectedSourceAssets, TArray<FLocalizedAssetsState> InLocalizedAssetsState);
 
 	/** Adds the Show Localized Assets sub-menu to a menu builder. */
-	void MakeShowLocalizedAssetSubMenu(UToolMenu* Menu, TArray<FLocalizedAssetsState> InLocalizedAssetsState);
+	UE_API void MakeShowLocalizedAssetSubMenu(UToolMenu* Menu, TArray<FLocalizedAssetsState> InLocalizedAssetsState);
 
 	/** Adds the Edit Localized Assets sub-menu to a menu builder. */
-	void MakeEditLocalizedAssetSubMenu(UToolMenu* Menu, TArray<FLocalizedAssetsState> InLocalizedAssetsState);
+	UE_API void MakeEditLocalizedAssetSubMenu(UToolMenu* Menu, TArray<FLocalizedAssetsState> InLocalizedAssetsState);
 
 	/** Create new localized assets for the given culture */
-	void ExecuteCreateLocalizedAsset(TSet<FSoftObjectPath> InSelectedSourceAssets, FLocalizedAssetsState InLocalizedAssetsStateForCulture);
+	UE_API void ExecuteCreateLocalizedAsset(TSet<FSoftObjectPath> InSelectedSourceAssets, FLocalizedAssetsState InLocalizedAssetsStateForCulture);
 
 	/** Find the given assets in the Content Browser */
-	void ExecuteFindInAssetTree(TArray<FSoftObjectPath> InAssets);
+	UE_API void ExecuteFindInAssetTree(TArray<FSoftObjectPath> InAssets);
 
 	/** Open the given assets in their respective editors */
-	void ExecuteOpenEditorsForAssets(TArray<FSoftObjectPath> InAssets);
+	UE_API void ExecuteOpenEditorsForAssets(TArray<FSoftObjectPath> InAssets);
 
 	/** Adds asset documentation menu options to a menu builder. Returns true if any options were added. */
-	bool AddDocumentationMenuOptions(UToolMenu* Menu);
+	UE_API bool AddDocumentationMenuOptions(UToolMenu* Menu);
 	
 	/** Creates a sub-menu of Chunk IDs that are are assigned to all selected assets */
-	void MakeChunkIDListMenu(UToolMenu* Menu);
+	UE_API void MakeChunkIDListMenu(UToolMenu* Menu);
 
 	/** Handler for when create using asset is selected */
-	void ExecuteCreateBlueprintUsing();
+	UE_API void ExecuteCreateBlueprintUsing();
 
 	/** Handler for when "Find in World" is selected */
-	void ExecuteFindAssetInWorld();
+	UE_API void ExecuteFindAssetInWorld();
 
 	/** Handler for when "Property Matrix..." is selected */
-	void ExecutePropertyMatrix();
+	UE_API void ExecutePropertyMatrix();
 
 	/** Handler for when "Show MetaData" is selected */
-	void ExecuteShowAssetMetaData();
+	UE_API void ExecuteShowAssetMetaData();
 
 	/** Handler for when "Diff Selected" is selected */
-	void ExecuteDiffSelected() const;
+	UE_API void ExecuteDiffSelected() const;
 
 	/** Handler for Consolidate */
-	void ExecuteConsolidate();
+	UE_API void ExecuteConsolidate();
 
 	/** Handler for Capture Thumbnail */
-	void ExecuteCaptureThumbnail();
+	UE_API void ExecuteCaptureThumbnail();
 
 	/** Handler for Clear Thumbnail */
-	void ExecuteClearThumbnail();
+	UE_API void ExecuteClearThumbnail();
 
 	/** Handler for when "Migrate Asset" is selected */
-	void ExecuteMigrateAsset();
+	UE_API void ExecuteMigrateAsset();
 
 	/** Handler for GoToAssetCode */
-	void ExecuteGoToCodeForAsset(UClass* SelectedClass);
+	UE_API void ExecuteGoToCodeForAsset(UClass* SelectedClass);
 
 	/** Handler for GoToAssetDocs */
-	void ExecuteGoToDocsForAsset(UClass* SelectedClass);
+	UE_API void ExecuteGoToDocsForAsset(UClass* SelectedClass);
 
 	/** Handler for GoToAssetDocs */
-	void ExecuteGoToDocsForAsset(UClass* SelectedClass, const FString ExcerptSection);
+	UE_API void ExecuteGoToDocsForAsset(UClass* SelectedClass, const FString ExcerptSection);
 	
 	/** Handler for resetting the localization ID of the current selection */
-	void ExecuteResetLocalizationId();
+	UE_API void ExecuteResetLocalizationId();
 
 	/** Handler for showing the cached list of localized texts stored in the package header */
-	void ExecuteShowLocalizationCache(const FAssetData InAsset, const FString InPackageFilename);
+	UE_API void ExecuteShowLocalizationCache(const FAssetData InAsset, const FString InPackageFilename);
 
 	/** Handler for Export */
-	void ExecuteExport();
+	UE_API void ExecuteExport();
 
 	/** Handler for Bulk Export */
-	void ExecuteBulkExport();
+	UE_API void ExecuteBulkExport();
 
 	/** Handler to assign ChunkID to a selection of assets */
-	void ExecuteAssignChunkID();
+	UE_API void ExecuteAssignChunkID();
 
 	/** Handler to remove all ChunkID assignments from a selection of assets */
-	void ExecuteRemoveAllChunkID();
+	UE_API void ExecuteRemoveAllChunkID();
 
 	/** Handler to remove a single ChunkID assignment from a selection of assets */
-	void ExecuteRemoveChunkID(int32 ChunkID);
+	UE_API void ExecuteRemoveChunkID(int32 ChunkID);
 
 	/** Handler to export the selected asset(s) to experimental text format */
-	void ExportSelectedAssetsToText();
+	UE_API void ExportSelectedAssetsToText();
 
 	/** Handler to export the selected asset(s) to experimental text format */
-	void ViewSelectedAssetAsText();
-	bool CanViewSelectedAssetAsText() const;
+	UE_API void ViewSelectedAssetAsText();
+	UE_API bool CanViewSelectedAssetAsText() const;
 
 	/** Run the rountrip test on this asset */
-	void DoTextFormatRoundtrip();
+	UE_API void DoTextFormatRoundtrip();
 
 	/** Handler to check if we can create blueprint using selected asset */
-	bool CanExecuteCreateBlueprintUsing() const;
+	UE_API bool CanExecuteCreateBlueprintUsing() const;
 
 	/** Handler to check to see if a find in world command is allowed */
-	bool CanExecuteFindAssetInWorld() const;
+	UE_API bool CanExecuteFindAssetInWorld() const;
 
 	/** Handler to check to see if a properties command is allowed */
-	bool CanExecuteProperties() const;
+	UE_API bool CanExecuteProperties() const;
 
 	/** Handler to check to see if a property matrix command is allowed */
-	bool CanExecutePropertyMatrix(FText& OutErrorMessage) const;
-	bool CanExecutePropertyMatrix() const;
+	UE_API bool CanExecutePropertyMatrix(FText& OutErrorMessage) const;
+	UE_API bool CanExecutePropertyMatrix() const;
 
 	/** Handler to check to see if "Capture Thumbnail" can be executed */
-	bool CanExecuteCaptureThumbnail() const;
+	UE_API bool CanExecuteCaptureThumbnail() const;
 
 	/** Handler to check to see if "Clear Thumbnail" should be visible */
-	bool CanClearCustomThumbnails() const;
+	UE_API bool CanClearCustomThumbnails() const;
 
 	/** Initializes some variable used to in "CanExecute" checks that won't change at runtime or are too expensive to check every frame. */
-	void CacheCanExecuteVars();
+	UE_API void CacheCanExecuteVars();
 
 	/** Helper function to gather the package names of all selected assets */
-	void GetSelectedPackageNames(TArray<FString>& OutPackageNames) const;
+	UE_API void GetSelectedPackageNames(TArray<FString>& OutPackageNames) const;
 
 	/** Helper function to gather the packages containing all selected assets */
-	void GetSelectedPackages(TArray<UPackage*>& OutPackages) const;
+	UE_API void GetSelectedPackages(TArray<UPackage*>& OutPackages) const;
 
 	/** Update internal state logic */
-	void OnChunkIDAssignChanged(int32 ChunkID);
+	UE_API void OnChunkIDAssignChanged(int32 ChunkID);
 
 	/** Gets the current value of the ChunkID entry box */
-	TOptional<int32> GetChunkIDSelection() const;
+	UE_API TOptional<int32> GetChunkIDSelection() const;
 
 	/** Handles when the Assign chunkID dialog OK button is clicked */
-	FReply OnChunkIDAssignCommit(TSharedPtr<SWindow> Window);
+	UE_API FReply OnChunkIDAssignCommit(TSharedPtr<SWindow> Window);
 
 	/** Handles when the Assign chunkID dialog Cancel button is clicked */
-	FReply OnChunkIDAssignCancel(TSharedPtr<SWindow> Window);
+	UE_API FReply OnChunkIDAssignCancel(TSharedPtr<SWindow> Window);
 
 	/** Generates tooltip for the Property Matrix menu option */
-	FText GetExecutePropertyMatrixTooltip() const;
+	UE_API FText GetExecutePropertyMatrixTooltip() const;
 
 	/** Generates a list of selected assets in the content browser */
-	void GetSelectedAssets(TArray<UObject*>& Assets, bool SkipRedirectors) const;
+	UE_API void GetSelectedAssets(TArray<UObject*>& Assets, bool SkipRedirectors) const;
 
 	/** Generates a list of selected assets in the content browser, and returns the asset data so you do not have to load them */
-	void GetSelectedAssetData(TArray<FAssetData>& AssetDataList, bool SkipRedirectors) const;
+	UE_API void GetSelectedAssetData(TArray<FAssetData>& AssetDataList, bool SkipRedirectors) const;
 
 private:
 	TArray<FAssetData> SelectedAssets;
@@ -249,3 +251,5 @@ private:
 	/** */
 	int32 ChunkIDSelected = 0;
 };
+
+#undef UE_API
diff --git a/Engine/Plugins/Editor/ContentBrowser/ContentBrowserAssetDataSource/Source/ContentBrowserAssetDataSource/Public/AssetFolderContextMenu.h b/Engine/Plugins/Editor/ContentBrowser/ContentBrowserAssetDataSource/Source/ContentBrowserAssetDataSource/Public/AssetFolderContextMenu.h
--- a/Engine/Plugins/Editor/ContentBrowser/ContentBrowserAssetDataSource/Source/ContentBrowserAssetDataSource/Public/AssetFolderContextMenu.h
+++ b/Engine/Plugins/Editor/ContentBrowser/ContentBrowserAssetDataSource/Source/ContentBrowserAssetDataSource/Public/AssetFolderContextMenu.h
@@ -4,13 +4,15 @@
 
 #include "Templates/SharedPointer.h"
 
+#define UE_API CONTENTBROWSERASSETDATASOURCE_API
+
 class UToolMenu;
 
-class CONTENTBROWSERASSETDATASOURCE_API FAssetFolderContextMenu : public TSharedFromThis<FAssetFolderContextMenu>
+class FAssetFolderContextMenu : public TSharedFromThis<FAssetFolderContextMenu>
 {
 public:
 	UE_DEPRECATED(5.1, "Use the new overload.")
-	void MakeContextMenu(
+	UE_API void MakeContextMenu(
 		UToolMenu* InMenu,
 		const TArray<FString>& InSelectedPackagePaths
 	);
@@ -21,7 +23,7 @@ public:
 	 * @param InSelectedPackagePaths The path of the selected folders
 	 * @param InSelectedPackages The package names of the selected asset item. Empty if the selection is only some folder items.
 	 */
-	void MakeContextMenu(
+	UE_API void MakeContextMenu(
 		UToolMenu* InMenu,
 		const TArray<FString>& InSelectedPackagePaths,
 		const TArray<FString>& InSelectedPackages
@@ -29,16 +31,16 @@ public:
 
 private:
 	/** Makes the asset tree context menu widget */
-	void AddMenuOptions(UToolMenu* Menu);
+	UE_API void AddMenuOptions(UToolMenu* Menu);
 
 	/** Handler for when "Migrate Folder" is selected */
-	void ExecuteMigrateFolder();
+	UE_API void ExecuteMigrateFolder();
 
 	/** Handler for when "Fix up Redirectors in Folder" is selected */
-	void ExecuteFixUpRedirectorsInFolder();
+	UE_API void ExecuteFixUpRedirectorsInFolder();
 
 	/** Gets the first selected path, if it exists */
-	FString GetFirstSelectedPath() const;
+	UE_API FString GetFirstSelectedPath() const;
 
 private:
 	TArray<FString> SelectedPaths;
@@ -46,3 +48,5 @@ private:
 	// Not null if the selection contained some assets also
 	TArray<FString> SelectedPackages;
 };
+
+#undef UE_API
diff --git a/Engine/Plugins/Editor/ContentBrowser/ContentBrowserAssetDataSource/Source/ContentBrowserAssetDataSource/Public/AssetPropertyTagCache.h b/Engine/Plugins/Editor/ContentBrowser/ContentBrowserAssetDataSource/Source/ContentBrowserAssetDataSource/Public/AssetPropertyTagCache.h
--- a/Engine/Plugins/Editor/ContentBrowser/ContentBrowserAssetDataSource/Source/ContentBrowserAssetDataSource/Public/AssetPropertyTagCache.h
+++ b/Engine/Plugins/Editor/ContentBrowser/ContentBrowserAssetDataSource/Source/ContentBrowserAssetDataSource/Public/AssetPropertyTagCache.h
@@ -4,10 +4,12 @@
 
 #include "UObject/Object.h"
 
+#define UE_API CONTENTBROWSERASSETDATASOURCE_API
+
 /**
  * Singleton cache of property tags data for each asset class.
  */
-class CONTENTBROWSERASSETDATASOURCE_API FAssetPropertyTagCache
+class FAssetPropertyTagCache
 {
 public:
 	struct FPropertyTagCache
@@ -56,19 +58,19 @@ public:
 	};
 
 	/** Get the singleton instance */
-	static FAssetPropertyTagCache& Get();
+	static UE_API FAssetPropertyTagCache& Get();
 
 	/** Try and populate the cache for the given class if it is loaded */
-	void TryCacheClass(FTopLevelAssetPath InClassName);
+	UE_API void TryCacheClass(FTopLevelAssetPath InClassName);
 
 	/** Try and populate the cache with classes which were not reigstered when TryCacheClass was called. */
-	void CachePendingClasses();
+	UE_API void CachePendingClasses();
 
 	/** Get (or populate) the cache for the given asset class */
-	const FClassPropertyTagCache& GetCacheForClass(FTopLevelAssetPath InClassName);
+	UE_API const FClassPropertyTagCache& GetCacheForClass(FTopLevelAssetPath InClassName);
 
 	/** Get the cache for the given asset class if it has been created */
-	const FClassPropertyTagCache* FindCacheForClass(FTopLevelAssetPath InClassName);
+	UE_API const FClassPropertyTagCache* FindCacheForClass(FTopLevelAssetPath InClassName);
 
 private:
 	mutable FRWLock Lock;
@@ -82,3 +84,5 @@ private:
 	/** Classes we'd like to register but were not loaded when an asset of that type was last scanned */
 	TSet<FTopLevelAssetPath> PendingClasses;
 };
+
+#undef UE_API
