Fixed UContentBrowserAliasDataSource::EnumerateItemsAtPath not working for files, which was preventing users from pasting aliased object paths into the content browser navigation bar to navigate to them.

diff --git a/Engine/Plugins/Editor/ContentBrowser/ContentBrowserAliasDataSource/Source/ContentBrowserAliasDataSource/Private/ContentBrowserAliasDataSource.cpp b/Engine/Plugins/Editor/ContentBrowser/ContentBrowserAliasDataSource/Source/ContentBrowserAliasDataSource/Private/ContentBrowserAliasDataSource.cpp
--- a/Engine/Plugins/Editor/ContentBrowser/ContentBrowserAliasDataSource/Source/ContentBrowserAliasDataSource/Private/ContentBrowserAliasDataSource.cpp
+++ b/Engine/Plugins/Editor/ContentBrowser/ContentBrowserAliasDataSource/Source/ContentBrowserAliasDataSource/Private/ContentBrowserAliasDataSource.cpp
@@ -263,32 +263,14 @@ void UContentBrowserAliasDataSource::EnumerateItemsAtPath(const FName InPath, co
 
 	if (EnumHasAnyFlags(InItemTypeFilter, EContentBrowserItemTypeFilter::IncludeFiles))
 	{
-		TSet<FName> Paths;
-		// Return all assets for this path and its subpaths
-		Paths.Add(InternalPath);
-		PathTree.GetSubPaths(InternalPath, Paths, true);
-
-		bool bIsAlreadyInSet = false;
-		for (const FName Path : Paths)
+		if (const FContentBrowserUniqueAlias* Alias = AliasForInternalPath.Find(InternalPath))
 		{
-			const TArray<FContentBrowserUniqueAlias>* Aliases = AliasesInPackagePath.Find(Path);
-			if (Aliases)
+			FContentBrowserItemData Item = CreateAssetFileItem(*Alias);
+			if (ensureMsgf(Item.IsValid(), TEXT("Could not find alias data for: %s"), *InternalPath.ToString()))
 			{
-				for (const FContentBrowserUniqueAlias& Alias : *Aliases)
-				{
-					AlreadyAddedOriginalAssets.Add(AllAliases[Alias].AssetData.GetSoftObjectPath(), &bIsAlreadyInSet);
-					if (!bIsAlreadyInSet)
-					{
-						if (!InCallback(CreateAssetFileItem(Alias)))
-						{
-							AlreadyAddedOriginalAssets.Reset();
-							return;
-						}
-					}
-				}
+				InCallback(MoveTemp(Item));
 			}
 		}
-		AlreadyAddedOriginalAssets.Reset();
 	}
 }
 
@@ -359,11 +341,13 @@ TArray<FContentBrowserItemPath> UContentBrowserAliasDataSource::GetAliasesForPat
 
 	if (const TArray<FName>* FoundAliases = AliasesForObjectPath.Find(ObjectPath))
 	{
+		OutAliases.Reserve(FoundAliases->Num());
+
 		for (const FName Alias : *FoundAliases)
 		{
 			if (const FAliasData* AliasData = AllAliases.Find(FContentBrowserUniqueAlias(ObjectPath, Alias)))
 			{
-				OutAliases.Add(FContentBrowserItemPath(AliasData->PackageName, EContentBrowserPathType::Internal));
+				OutAliases.Add(FContentBrowserItemPath(AliasData->InternalPath, EContentBrowserPathType::Internal));
 			}
 		}
 	}
@@ -469,7 +453,8 @@ void UContentBrowserAliasDataSource::AddAliasImpl(const FAssetData& Asset, const
 
 	AliasesInPackagePath.FindOrAdd(PackagePath).Add(UniqueAlias);
 	AliasesForObjectPath.FindOrAdd(Asset.GetSoftObjectPath()).Add(Alias);
-	AllAliases.Add(UniqueAlias, FAliasData(Asset, PackagePath, AliasDisplayNameOverride.IsEmpty() ? FText::AsCultureInvariant(*Tokens.Last()) : AliasDisplayNameOverride, bInIsFromMetaData));
+	FAliasData& AliasData = AllAliases.Add(UniqueAlias, FAliasData(Asset, PackagePath, AliasDisplayNameOverride.IsEmpty() ? FText::AsCultureInvariant(*Tokens.Last()) : AliasDisplayNameOverride, bInIsFromMetaData));
+	AliasForInternalPath.Add(AliasData.InternalPath, UniqueAlias);
 	QueueItemDataUpdate(FContentBrowserItemDataUpdate::MakeItemAddedUpdate(CreateAssetFileItem(UniqueAlias)));
 
 	// This logging might get out of control if there ends up being hundreds of thousands of aliases.
@@ -545,6 +530,8 @@ void UContentBrowserAliasDataSource::RemoveAlias(const FSoftObjectPath& ObjectPa
 			}
 		}
 
+		AliasForInternalPath.Remove(AliasData.InternalPath);
+
 		QueueItemDataUpdate(FContentBrowserItemDataUpdate::MakeItemRemovedUpdate(ItemData));
 
 		FNameBuilder AliasAsString(Alias);
@@ -1075,8 +1062,7 @@ FContentBrowserItemData UContentBrowserAliasDataSource::CreateAssetFileItem(cons
 	{
 		LLM_SCOPE(ELLMTag::UI);
 		FName VirtualizedPath;
-		FName InternalPath = *AliasData->ObjectPath.ToString();
-		TryConvertInternalPathToVirtual(InternalPath, VirtualizedPath);
+		TryConvertInternalPathToVirtual(AliasData->InternalPath, VirtualizedPath);
 
 		// Since AliasID is PackagePath/AssetName, AssetName should also be passed as the ItemName here. This provides the functionality of
 		// being able to have multiple aliases with the same display name, while still showing their original asset name in the tooltip.
@@ -1087,7 +1073,7 @@ FContentBrowserItemData UContentBrowserAliasDataSource::CreateAssetFileItem(cons
 			AliasData->AssetData.AssetName,
 			AliasData->AliasDisplayName,
 			MakeShared<FContentBrowserAliasItemDataPayload>(AliasData->AssetData, Alias),
-			{ InternalPath });
+			{ AliasData->InternalPath });
 	}
 	return FContentBrowserItemData();
 }
diff --git a/Engine/Plugins/Editor/ContentBrowser/ContentBrowserAliasDataSource/Source/ContentBrowserAliasDataSource/Public/ContentBrowserAliasDataSource.h b/Engine/Plugins/Editor/ContentBrowser/ContentBrowserAliasDataSource/Source/ContentBrowserAliasDataSource/Public/ContentBrowserAliasDataSource.h
--- a/Engine/Plugins/Editor/ContentBrowser/ContentBrowserAliasDataSource/Source/ContentBrowserAliasDataSource/Public/ContentBrowserAliasDataSource.h
+++ b/Engine/Plugins/Editor/ContentBrowser/ContentBrowserAliasDataSource/Source/ContentBrowserAliasDataSource/Public/ContentBrowserAliasDataSource.h
@@ -216,18 +216,21 @@ private:
 				PathBuilder << TEXT('/');
 				PathBuilder << AssetNameBuilder.ToView();
 				PackageName = FName(PathBuilder.ToView());
+
+				PathBuilder << TEXT('.');
+				PathBuilder << AssetNameBuilder.ToView();
+				InternalPath = FName(PathBuilder.ToView());
 			}
-			ObjectPath = FSoftObjectPath(FTopLevelAssetPath(PackageName, FName(AssetNameBuilder.ToView())));
 		}
 
 		/** The source asset for this alias */
 		FAssetData AssetData;
 		/** The folder path that contains the alias, /MyAliases */
 		FName PackagePath;
-		/** PackagePath/SourceAssetName, /MyAliases/SomeAsset  */
+		/** PackagePath/SourceAssetName, /MyAliases/SomeAsset_Hash  */
 		FName PackageName;
-		/** PackageName.SourceAssetName, /MyAliases/SomeAsset.SomeAsset */
-		FSoftObjectPath ObjectPath;
+		/** PackageName.SourceAssetName, /MyAliases/SomeAsset_Hash.SomeAsset_Hash */
+		FName InternalPath;
 		/** A non-unique display name for this alias */
 		FText AliasDisplayName;
 		/** Whether this alias was generated from package metadata or manually through the C++ interface */
@@ -243,6 +246,8 @@ private:
 	TMap<FSoftObjectPath, TArray<FName>> AliasesForObjectPath;
 	/** A list of alias paths to display for each folder, eg /Game/MyData/Aliases */
 	TMap<FName, TArray<FContentBrowserUniqueAlias>> AliasesInPackagePath;
+	/** The alias path to display for each file */
+	TMap<FName, FContentBrowserUniqueAlias> AliasForInternalPath;
 	/** Alias folder display names keyed against their alias path, eg /Game/MyData/Aliases/ */
 	TMap<FName, FText> AliasFolderDisplayNames;
 	/** A set used for removing duplicate aliases in the same query, stored here to avoid constant reallocation */
