Content browser filtering de-duplication

`UContentBrowserAliasDataSource` used to do de-duplication when enumerating items.

Problem: when using a search query (like Path="Something"), one of the discarded duplicate alias could match the query. So depending on the order in which the data source enumerates through aliases, the content browser may give different results. `UContentBrowserAliasDataSource` discards duplicates too early, as one of the duplicates would eventually get through the search query and frontend filtering.

I now filter duplicates after frontend filtering and the search query. I do that for item aliases, through the new `GetItemAliasedPath` method.

#Synced-CL 48418225

diff --git a/Engine/Plugins/Editor/ContentBrowser/ContentBrowserAliasDataSource/Source/ContentBrowserAliasDataSource/Private/ContentBrowserAliasDataSource.cpp b/Engine/Plugins/Editor/ContentBrowser/ContentBrowserAliasDataSource/Source/ContentBrowserAliasDataSource/Private/ContentBrowserAliasDataSource.cpp
--- a/Engine/Plugins/Editor/ContentBrowser/ContentBrowserAliasDataSource/Source/ContentBrowserAliasDataSource/Private/ContentBrowserAliasDataSource.cpp
+++ b/Engine/Plugins/Editor/ContentBrowser/ContentBrowserAliasDataSource/Source/ContentBrowserAliasDataSource/Private/ContentBrowserAliasDataSource.cpp
@@ -195,7 +195,6 @@ void UContentBrowserAliasDataSource::EnumerateItemsMatchingFilter(const FContent
 
 	if (EnumHasAnyFlags(InFilter.ItemTypeFilter, EContentBrowserItemTypeFilter::IncludeFiles) && !AssetDataFilter->bFilterExcludesAllAssets && !AssetDataFilter->InclusiveFilter.IsEmpty())
 	{
-		bool bIsAlreadyInSet = false;
 		if (AssetDataFilter->InclusiveFilter.PackagePaths.Num() > 0)
 		{
 			// Find all aliases for each requested package path and check if it passes both
@@ -210,14 +209,9 @@ void UContentBrowserAliasDataSource::EnumerateItemsMatchingFilter(const FContent
 						const FAliasData& AliasData = AllAliases[Alias];
 						if (DoesAliasPassFilter(AliasData, *AssetDataFilter))
 						{
-							AlreadyAddedOriginalAssets.Add(AliasData.AssetData.GetSoftObjectPath(), &bIsAlreadyInSet);
-							if (!bIsAlreadyInSet)
+							if (!InCallback(CreateAssetFileItem(Alias)))
 							{
-								if (!InCallback(CreateAssetFileItem(Alias)))
-								{
-									AlreadyAddedOriginalAssets.Reset();
-									return;
-								}
+								return;
 							}
 						}
 					}
@@ -231,19 +225,13 @@ void UContentBrowserAliasDataSource::EnumerateItemsMatchingFilter(const FContent
 			{
 				if (DoesAliasPassFilter(Pair.Value, *AssetDataFilter))
 				{
-					AlreadyAddedOriginalAssets.Add(Pair.Value.AssetData.GetSoftObjectPath(), &bIsAlreadyInSet);
-					if (!bIsAlreadyInSet)
+					if (!InCallback(CreateAssetFileItem(Pair.Key)))
 					{
-						if (!InCallback(CreateAssetFileItem(Pair.Key)))
-						{
-							AlreadyAddedOriginalAssets.Reset();
-							return;
-						}
+						return;
 					}
 				}
 			}
 		}
-		AlreadyAddedOriginalAssets.Reset();
 	}
 }
 
@@ -423,6 +411,22 @@ TArray<FContentBrowserItemPath> UContentBrowserAliasDataSource::GetAliasesForPat
 	return OutAliases;
 }
 
+bool UContentBrowserAliasDataSource::GetItemAliasedPath(const FContentBrowserItemData& InItemData, FSoftObjectPath& OutObjectPath) const
+{
+	if (InItemData.GetOwnerDataSource() == this && InItemData.IsFile())
+	{
+		if (const TSharedPtr<const FContentBrowserAliasItemDataPayload> AliasPayload = StaticCastSharedPtr<const FContentBrowserAliasItemDataPayload>(InItemData.GetPayload()))
+		{
+			if (AliasPayload->Alias.Key.IsValid())
+			{
+				OutObjectPath = AliasPayload->Alias.Key;
+				return true;
+			}
+		}
+	}
+	return Super::GetItemAliasedPath(InItemData, OutObjectPath);
+}
+
 bool UContentBrowserAliasDataSource::HasAliasesForPath(const FSoftObjectPath& ObjectPath) const
 {
 	return AliasesForObjectPath.Contains(ObjectPath);
@@ -644,7 +648,6 @@ void UContentBrowserAliasDataSource::RebuildAliases()
 	AliasesForObjectPath.Reset();
 	AliasesInPackagePath.Reset();
 	AliasFolderDisplayNames.Reset();
-	AlreadyAddedOriginalAssets.Reset();
 	FilterCache.Reset();
 	RootPathVirtualTree.Reset();
 
diff --git a/Engine/Plugins/Editor/ContentBrowser/ContentBrowserAliasDataSource/Source/ContentBrowserAliasDataSource/Public/ContentBrowserAliasDataSource.h b/Engine/Plugins/Editor/ContentBrowser/ContentBrowserAliasDataSource/Source/ContentBrowserAliasDataSource/Public/ContentBrowserAliasDataSource.h
--- a/Engine/Plugins/Editor/ContentBrowser/ContentBrowserAliasDataSource/Source/ContentBrowserAliasDataSource/Public/ContentBrowserAliasDataSource.h
+++ b/Engine/Plugins/Editor/ContentBrowser/ContentBrowserAliasDataSource/Source/ContentBrowserAliasDataSource/Public/ContentBrowserAliasDataSource.h
@@ -69,6 +69,7 @@ public:
 
 	using UContentBrowserDataSource::GetAliasesForPath;
 	UE_API virtual TArray<FContentBrowserItemPath> GetAliasesForPath(const FSoftObjectPath& InInternalPath) const override;
+	UE_API virtual bool GetItemAliasedPath(const FContentBrowserItemData& InItemData, FSoftObjectPath& OutObjectPath) const override;
 
 	UE_API bool HasAliasesForPath(const FSoftObjectPath& InInternalPath) const;
 
@@ -253,8 +254,6 @@ private:
 	TMap<FName, FContentBrowserUniqueAlias> AliasForInternalPath;
 	/** Alias folder display names keyed against their alias path, eg /Game/MyData/Aliases/ */
 	TMap<FName, FText> AliasFolderDisplayNames;
-	/** A set used for removing duplicate aliases in the same query, stored here to avoid constant reallocation */
-	TSet<FSoftObjectPath> AlreadyAddedOriginalAssets;
 
 	UContentBrowserAssetDataSource::FAssetDataSourceFilterCache FilterCache;
 
diff --git a/Engine/Source/Editor/ContentBrowser/Private/SAssetView.cpp b/Engine/Source/Editor/ContentBrowser/Private/SAssetView.cpp
--- a/Engine/Source/Editor/ContentBrowser/Private/SAssetView.cpp
+++ b/Engine/Source/Editor/ContentBrowser/Private/SAssetView.cpp
@@ -2668,7 +2668,25 @@ void SAssetView::ProcessItemsPendingFilter(const double TickStartTime)
 	FAssetViewFrontendFilterHelper FrontendFilterHelper(this);
 	const bool bFlushAllPendingItems = TickStartTime < 0;
 	int32 OldCount = FilteredAssetItems.Num();
-	Items->UpdateItemFiltering(FrontendFilterHelper, bFlushAllPendingItems ? MAX_dbl : TickStartTime + MaxSecondsPerFrame, FilteredAssetItems);
+	TArray<TSharedPtr<FAssetViewItem>> FilterResults;
+
+	Items->UpdateItemFiltering(FrontendFilterHelper, bFlushAllPendingItems ? MAX_dbl : TickStartTime + MaxSecondsPerFrame, FilterResults);
+
+	// exclude duplicates
+	for (TSharedPtr<FAssetViewItem>& Result : FilterResults)
+	{
+		bool bAlreadyAdded = false;
+		const FContentBrowserItem& Item = Result->GetItem();
+		
+		if (FSoftObjectPath ItemPath; Item.IsValid() && Item.GetItemAliasedPath(ItemPath))
+		{
+			AlreadyAddedOriginalAssets.Add(ItemPath, &bAlreadyAdded);
+		}
+		if (!bAlreadyAdded)
+		{
+			FilteredAssetItems.Add(MoveTemp(Result));
+		}
+	}
 
 	if (CurrentViewType == EAssetViewType::Column)
 	{
@@ -3758,6 +3776,7 @@ void SAssetView::RefreshFilteredItems()
 	OnInterruptFiltering();
 	
 	FilteredAssetItems.Reset();
+	AlreadyAddedOriginalAssets.Reset();
 	FilteredAssetItemTypeCounts.Reset();
 	RelevantThumbnails.Reset();
 
diff --git a/Engine/Source/Editor/ContentBrowser/Public/SAssetView.h b/Engine/Source/Editor/ContentBrowser/Public/SAssetView.h
--- a/Engine/Source/Editor/ContentBrowser/Public/SAssetView.h
+++ b/Engine/Source/Editor/ContentBrowser/Public/SAssetView.h
@@ -1391,6 +1391,9 @@ private:
 	/** The actual widget for the custom view */
 	TSharedPtr<SWidget> CustomView;
 	
+	/** Existing item paths to avoid duplicates */
+	TSet<FSoftObjectPath> AlreadyAddedOriginalAssets;
+	
 	/*
 	 * Telemetry-related fields
 	 */ 
diff --git a/Engine/Source/Editor/ContentBrowserData/Private/ContentBrowserDataSource.cpp b/Engine/Source/Editor/ContentBrowserData/Private/ContentBrowserDataSource.cpp
--- a/Engine/Source/Editor/ContentBrowserData/Private/ContentBrowserDataSource.cpp
+++ b/Engine/Source/Editor/ContentBrowserData/Private/ContentBrowserDataSource.cpp
@@ -259,6 +259,11 @@ TArray<FContentBrowserItemPath> UContentBrowserDataSource::GetAliasesForPath(con
 	return TArray<FContentBrowserItemPath>();
 }
 
+bool UContentBrowserDataSource::GetItemAliasedPath(const FContentBrowserItemData& InItemData, FSoftObjectPath& OutObjectPath) const
+{
+	return false;
+}
+
 TArray<FContentBrowserItemPath> UContentBrowserDataSource::GetAliasesForPath(FName InInternalPath) const
 {
 	return GetAliasesForPath(FSoftObjectPath(InInternalPath.ToString()));
diff --git a/Engine/Source/Editor/ContentBrowserData/Private/ContentBrowserItem.cpp b/Engine/Source/Editor/ContentBrowserData/Private/ContentBrowserItem.cpp
--- a/Engine/Source/Editor/ContentBrowserData/Private/ContentBrowserItem.cpp
+++ b/Engine/Source/Editor/ContentBrowserData/Private/ContentBrowserItem.cpp
@@ -648,6 +648,11 @@ bool FContentBrowserItem::Legacy_TryGetAssetData(FAssetData& OutAssetData) const
 	return FContentBrowserItemHelper::CallDataSourceImpl<UContentBrowserDataSource>(*this, UE_PROJECTION_MEMBER(UContentBrowserDataSource, Legacy_TryGetAssetData), OutAssetData);
 }
 
+bool FContentBrowserItem::GetItemAliasedPath(FSoftObjectPath& OutObjectPath) const
+{
+	return FContentBrowserItemHelper::CallDataSourceImpl<UContentBrowserDataSource>(*this, UE_PROJECTION_MEMBER(UContentBrowserDataSource, GetItemAliasedPath), OutObjectPath);
+}
+
 FContentBrowserItemPath FContentBrowserItem::GetItemPath() const
 {
 	const FName VirtualPath = GetVirtualPath();
diff --git a/Engine/Source/Editor/ContentBrowserData/Public/ContentBrowserDataSource.h b/Engine/Source/Editor/ContentBrowserData/Public/ContentBrowserDataSource.h
--- a/Engine/Source/Editor/ContentBrowserData/Public/ContentBrowserDataSource.h
+++ b/Engine/Source/Editor/ContentBrowserData/Public/ContentBrowserDataSource.h
@@ -271,10 +271,19 @@ public:
 	/**
 	 * Get a list of other paths that the data source may be using to represent a specific path
 	 *
-	 * @param The internal path (or object path) of an asset to get aliases for
+	 * @param InInternalPath The internal path (or object path) of an asset to get aliases for
 	 * @return All alternative paths that represent the input path (not including the input path itself)
 	 */
 	UE_API virtual TArray<FContentBrowserItemPath> GetAliasesForPath(const FSoftObjectPath& InInternalPath) const;
+	
+	/**
+	 * Get the source asset path of an alias.
+	 * 
+	 * @param InItemData The content browser item of an alias
+	 * @param OutObjectPath The soft object path to fill with the result
+	 * @return True if @code InItemData @endcode is an alias and a path was found, false otherwise
+	 */
+	UE_API virtual bool GetItemAliasedPath(const FContentBrowserItemData& InItemData, FSoftObjectPath& OutObjectPath) const;
 
 	UE_DEPRECATED(5.1, "FNames containing full asset paths are deprecated. Use FSoftObjectPath instead.")
 	UE_API TArray<FContentBrowserItemPath> GetAliasesForPath(FName InInternalPath) const;
diff --git a/Engine/Source/Editor/ContentBrowserData/Public/ContentBrowserItem.h b/Engine/Source/Editor/ContentBrowserData/Public/ContentBrowserItem.h
--- a/Engine/Source/Editor/ContentBrowserData/Public/ContentBrowserItem.h
+++ b/Engine/Source/Editor/ContentBrowserData/Public/ContentBrowserItem.h
@@ -497,6 +497,14 @@ public:
 	 */
 	UE_API bool Legacy_TryGetAssetData(FAssetData& OutAssetData) const;
 
+	/**
+	 * Attempt to retrieve the source asset path of an alias.
+	 * 
+	 * @param OutObjectPath The soft object path to fill with the result
+	 * @return True if @code this @endcode is an alias and a path was found, false otherwise
+	 */
+	UE_API bool GetItemAliasedPath(FSoftObjectPath& OutObjectPath) const;
+
 	/**
 	 * Creates a ContentBrowserItemPath, potentially looking up internal paths.
 	 * @return The Item Path.
