Fixes the Content Browser to not show the `_Verse` folders and `_Verse` items coming from cooked plugin data.

#rnx

#Synced-CL 49994500

#Synced-CL 49994500

#Synced-CL 49994500

diff --git a/Engine/Plugins/Editor/ContentBrowser/ContentBrowserAliasDataSource/Source/ContentBrowserAliasDataSource/Private/ContentBrowserAliasDataSource.cpp b/Engine/Plugins/Editor/ContentBrowser/ContentBrowserAliasDataSource/Source/ContentBrowserAliasDataSource/Private/ContentBrowserAliasDataSource.cpp
--- a/Engine/Plugins/Editor/ContentBrowser/ContentBrowserAliasDataSource/Source/ContentBrowserAliasDataSource/Private/ContentBrowserAliasDataSource.cpp
+++ b/Engine/Plugins/Editor/ContentBrowser/ContentBrowserAliasDataSource/Source/ContentBrowserAliasDataSource/Private/ContentBrowserAliasDataSource.cpp
@@ -15,6 +15,7 @@
 #include "ICollectionContainer.h"
 #include "IContentBrowserDataModule.h"
 #include "Misc/PackageName.h"
+#include "VerseVM/VVMNames.h"
 
 #include UE_INLINE_GENERATED_CPP_BY_NAME(ContentBrowserAliasDataSource)
 
@@ -386,6 +387,12 @@ bool UContentBrowserAliasDataSource::DoesAliasPassFilter(const FAliasData& Alias
 	// Create a fake asset data using the alias path instead of the asset's original path
 	// AssetRegistry->IsAssetIncludedByFilter is effectively a static function and does not actually use AssetRegistry data
 	FAssetData AliasAssetData(AliasData.PackageName, AliasData.PackagePath, AliasData.AssetData.AssetName, AliasData.AssetData.AssetClassPath, AliasData.AssetData.TagsAndValues.CopyMap());
+	
+	static const FString VerseFolderName = Verse::Names::GetVerseSubPath<TCHAR>();
+	if (FNameBuilder(AliasData.AssetData.PackageName).ToView().EndsWith(VerseFolderName))
+	{
+		return false;
+	}
 
 	return (Filter.InclusiveFilter.IsEmpty() || AssetRegistry->IsAssetIncludedByFilter(AliasAssetData, Filter.InclusiveFilter)) // Passes Inclusive
 		&& (Filter.ExclusiveFilter.IsEmpty() || !AssetRegistry->IsAssetIncludedByFilter(AliasAssetData, Filter.ExclusiveFilter)); // Passes Exclusive
diff --git a/Engine/Source/Editor/ContentBrowserData/Private/ContentBrowserDataUtils.cpp b/Engine/Source/Editor/ContentBrowserData/Private/ContentBrowserDataUtils.cpp
--- a/Engine/Source/Editor/ContentBrowserData/Private/ContentBrowserDataUtils.cpp
+++ b/Engine/Source/Editor/ContentBrowserData/Private/ContentBrowserDataUtils.cpp
@@ -18,6 +18,7 @@
 #include "Settings/ContentBrowserSettings.h"
 #include "Templates/SharedPointer.h"
 #include "UObject/UObjectGlobals.h"
+#include "VerseVM/VVMNames.h"
 
 #define LOCTEXT_NAMESPACE "ContentBrowserAssetDataUtils"
 
@@ -84,6 +85,7 @@ bool ContentBrowserDataUtils::PathPassesAttributeFilter(const FStringView InPath
 	static const FString ExternalActorsFolderName = FPackagePath::GetExternalActorsFolderName();
 	static const FString ExternalObjectsFolderName = FPackagePath::GetExternalObjectsFolderName();
 	static const FString DeveloperPathWithoutSlash = FPackageName::FilenameToLongPackageName(FPaths::GameDevelopersDir()).LeftChop(1);
+	static const FString VerseFolderName = Verse::Names::GetVerseSubPath<TCHAR>();
 	static const int32 DevelopersFolderDepth = ContentBrowserDataUtils::CalculateFolderDepthOfPath(DeveloperPathWithoutSlash);
 	const int32 MaxFolderDepthToCheck = ContentBrowserDataUtils::GetMaxFolderDepthRequiredForAttributeFilter();
 
@@ -197,6 +199,13 @@ bool ContentBrowserDataUtils::PathPassesAttributeFilter(const FStringView InPath
 				return false;
 			}
 		}
+
+		// NOTE: (yiliang.siew) We don't want to display the implicit `_Verse` folder that only contains cooked `UVerseClass`es
+		// that are of no use to the end-user in the Content Browser.
+		if (InPath.EndsWith(VerseFolderName))
+		{
+			return false;
+		}
 	}
 
 	return true;
